UE4: Fixed a crash at shutdown in DirectoryWatcher. Cancelled I/O requests must be completed before closing the handle and deleting the request.

diff --git a/Engine/Source/Developer/DirectoryWatcher/Private/Windows/DirectoryWatchRequestWindows.cpp b/Engine/Source/Developer/DirectoryWatcher/Private/Windows/DirectoryWatchRequestWindows.cpp
--- a/Engine/Source/Developer/DirectoryWatcher/Private/Windows/DirectoryWatchRequestWindows.cpp
+++ b/Engine/Source/Developer/DirectoryWatcher/Private/Windows/DirectoryWatchRequestWindows.cpp
@@ -139,6 +139,9 @@ void FDirectoryWatchRequestWindows::EndWatchRequest()
 #else
 			CancelIo(DirectoryHandle);
 #endif
+			// Wait for cancel operation to finish before closing the handle
+			WaitForSingleObjectEx(DirectoryHandle, 1000, true);
+			
 			::CloseHandle(DirectoryHandle);
 			DirectoryHandle = INVALID_HANDLE_VALUE;
 		}
