Linux: Fixing broken compilation of DirWatcher.

(how could I miss that?)

diff --git a/Engine/Source/Developer/DirectoryWatcher/Private/Linux/DirectoryWatchRequestLinux.cpp b/Engine/Source/Developer/DirectoryWatcher/Private/Linux/DirectoryWatchRequestLinux.cpp
--- a/Engine/Source/Developer/DirectoryWatcher/Private/Linux/DirectoryWatchRequestLinux.cpp
+++ b/Engine/Source/Developer/DirectoryWatcher/Private/Linux/DirectoryWatchRequestLinux.cpp
@@ -48,7 +48,7 @@ bool FDirectoryWatchRequestLinux::Init(const FString& InDirectory)
 
     IFileManager::Get().FindFilesRecursive(AllFiles, *FullPath, TEXT("*"), false, true);
 
-    WatchDescriptor = (int*)FMemory::Alloc((AllFiles.Num()+1) * sizeof(int));
+    WatchDescriptor = (int*)FMemory::Malloc((AllFiles.Num()+1) * sizeof(int));
     if (WatchDescriptor == NULL) 
     {
 		UE_LOG(LogDirectoryWatcher, Error, TEXT("Failed to allocate memory for WatchDescriptor"));
