Fixes to DirectoryWatcher from PR #306.

diff --git a/Engine/Source/Developer/DirectoryWatcher/Private/Linux/DirectoryWatchRequestLinux.cpp b/Engine/Source/Developer/DirectoryWatcher/Private/Linux/DirectoryWatchRequestLinux.cpp
--- a/Engine/Source/Developer/DirectoryWatcher/Private/Linux/DirectoryWatchRequestLinux.cpp
+++ b/Engine/Source/Developer/DirectoryWatcher/Private/Linux/DirectoryWatchRequestLinux.cpp
@@ -6,18 +6,18 @@ FDirectoryWatchRequestLinux::FDirectoryWatchRequestLinux()
 :	bRunning(false)
 ,	bEndWatchRequestInvoked(false)
 {
-    NotifyFilter = IN_CREATE | IN_MOVE | IN_MODIFY | IN_DELETE;
+	NotifyFilter = IN_CREATE | IN_MOVE | IN_MODIFY | IN_DELETE;
 }
 
 FDirectoryWatchRequestLinux::~FDirectoryWatchRequestLinux()
 {
-    Shutdown();
+	Shutdown();
 }
 
 void FDirectoryWatchRequestLinux::Shutdown()
 {
-    free(WatchDescriptor);
-    close(FileDescriptor);
+	FMemory::Free(WatchDescriptor);
+	close(FileDescriptor);
 }
 
 bool FDirectoryWatchRequestLinux::Init(const FString& InDirectory)
@@ -28,6 +28,8 @@ bool FDirectoryWatchRequestLinux::Init(const FString& InDirectory)
 		return false;
 	}
 
+	Directory = InDirectory;
+
 	if (bRunning)
 	{
 		Shutdown();
@@ -38,47 +40,51 @@ bool FDirectoryWatchRequestLinux::Init(const FString& InDirectory)
 	// Make sure the path is absolute
 	const FString FullPath = FPaths::ConvertRelativePathToFull(InDirectory);
 
-    FileDescriptor = inotify_init1(IN_NONBLOCK | IN_CLOEXEC);
+	FileDescriptor = inotify_init1(IN_NONBLOCK | IN_CLOEXEC);
 
-    if (FileDescriptor == -1)
-    {
+	if (FileDescriptor == -1)
+	{
+		// Failed to init inotify
 		UE_LOG(LogDirectoryWatcher, Error, TEXT("Failed to init inotify"));
-    	return false;
-    }
+		return false;
+	}
 
-    IFileManager::Get().FindFilesRecursive(AllFiles, *FullPath, TEXT("*"), false, true);
+	IFileManager::Get().FindFilesRecursive(AllFiles, *FullPath, TEXT("*"), false, true);
 
-    WatchDescriptor = (int*)FMemory::Malloc((AllFiles.Num()+1) * sizeof(int));
-    if (WatchDescriptor == NULL) 
-    {
+	// Allocate memory for watch descriptors
+	SIZE_T AllocSize = AllFiles.Num()+1 * sizeof(int));
+	WatchDescriptor = reinterpret_cast<int*>(FMemory::Malloc(AllocSize);
+	if (WatchDescriptor == nullptr) 
+	{
 		UE_LOG(LogDirectoryWatcher, Error, TEXT("Failed to allocate memory for WatchDescriptor"));
-    	return false;
-    }
-
-    for (int32 FileIdx = 0; FileIdx < AllFiles.Num(); ++FileIdx)
-    {
-		const FString& FolderName = AllFiles[FileIdx];
-		WatchDescriptor[FileIdx] = inotify_add_watch(FileDescriptor, TCHAR_TO_UTF8(*FolderName), NotifyFilter); //FileIdx+1
-		if (WatchDescriptor[FileIdx] == -1) 
-		{
-			UE_LOG(LogDirectoryWatcher, Error, TEXT("inotify_add_watch cannot watch folder %s"), *FolderName);
-			return false;
-		}
-    }
+		return false;
+	}
+	FMemory::Memzero(WatchDescriptor, AllocSize);
 
-  	bRunning = true;
+	for (int32 FileIdx = 0; FileIdx < AllFiles.Num(); ++FileIdx)
+	{
+			const FString& FolderName = AllFiles[FileIdx];
+			WatchDescriptor[FileIdx] = inotify_add_watch(FileDescriptor, TCHAR_TO_UTF8(*FolderName), NotifyFilter); //FileIdx+1
+			if (WatchDescriptor[FileIdx] == -1) 
+			{
+				UE_LOG(LogDirectoryWatcher, Error, TEXT("inotify_add_watch cannot watch folder %s"), *FolderName);
+				return false;
+			}
+	}
 
-  	return true;
+	bRunning = true;
+
+	return true;
 }
 
-void FDirectoryWatchRequestLinux::AddDelegate( const IDirectoryWatcher::FDirectoryChanged& InDelegate )
+void FDirectoryWatchRequestLinux::AddDelegate(const IDirectoryWatcher::FDirectoryChanged& InDelegate)
 {
 	Delegates.Add(InDelegate);
 }
 
-bool FDirectoryWatchRequestLinux::RemoveDelegate( const IDirectoryWatcher::FDirectoryChanged& InDelegate )
+bool FDirectoryWatchRequestLinux::RemoveDelegate(const IDirectoryWatcher::FDirectoryChanged& InDelegate)
 {
-	if ( Delegates.Contains(InDelegate) )
+	if (Delegates.Contains(InDelegate))
 	{
 		Delegates.Remove(InDelegate);
 		return true;
@@ -100,10 +106,10 @@ void FDirectoryWatchRequestLinux::EndWatchRequest()
 
 void FDirectoryWatchRequestLinux::ProcessPendingNotifications()
 {
-    ProcessChanges();
+	ProcessChanges();
 
 	// Trigger all listening delegates with the files that have changed
-	if ( FileChanges.Num() > 0 )
+	if (FileChanges.Num() > 0)
 	{
 		for (int32 DelegateIdx = 0; DelegateIdx < Delegates.Num(); ++DelegateIdx)
 		{
@@ -116,57 +122,54 @@ void FDirectoryWatchRequestLinux::ProcessPendingNotifications()
 
 void FDirectoryWatchRequestLinux::ProcessChanges()
 {
-    uint8 Buf[BUF_LEN] __attribute__ ((aligned(__alignof__(inotify_event))));
+	uint8_t Buf[BUF_LEN] __attribute__ ((aligned(__alignof__(struct inotify_event))));
+	const struct inotify_event* Event;
+	ssize_t Len = 0;
+	uint8_t* Ptr = nullptr;
 
+	// Loop while events can be read from inotify file descriptor
 	for (;;)
 	{
-		ssize_t Len = read(FileDescriptor, Buf, sizeof(Buf));
+		FFileChangeData::EFileChangeAction Action = FFileChangeData::FCA_Unknown;
+		// Read some events
+		Len = read(FileDescriptor, Buf, sizeof Buf);
 		if (Len == -1 && errno != EAGAIN)
 		{
-			// TODO: handle errors
-			return;
+			UE_LOG(LogDirectoryWatcher, Fatal, TEXT("DirectoryWatchRequest ProcessChanges read error"));
+			// unreachable
+			break;
 		}
 
+		// If the non-blocking read() found no events to read, then it returns -1 with errno set to EAGAIN. 
 		if (Len <= 0)
 		{
 			break;
 		}
-		else
+
+		// Loop over all events in the buffer
+		for (Ptr = Buf; Ptr < Buf + Len;)
 		{
-			for (const uint8 * Ptr = Buf; Ptr < Buf + Len;)
+			Event = reinterpret_cast<const struct inotify_event *>(Ptr);
+
+			if ((Event->mask & IN_CREATE) || (Event->mask & IN_MOVED_TO))
+			{
+				Action = FFileChangeData::FCA_Added;
+			}
+
+			if (Event->mask & IN_MODIFY)
+			{
+				Action = FFileChangeData::FCA_Modified;
+			}
+
+			if ((Event->mask & IN_DELETE) || (Event->mask & IN_MOVED_FROM))
 			{
-				const inotify_event* Event = reinterpret_cast<const inotify_event *>(Ptr);
-				FFileChangeData::EFileChangeAction Action = FFileChangeData::FCA_Unknown;
-
-				// not sure if "else" is correct here, but we cannot OR actions that way
-				if ((Event->mask & IN_CREATE) || (Event->mask & IN_MOVED_TO))
-				{
-					Action = FFileChangeData::FCA_Added;
-				}
-				else if (Event->mask & IN_MODIFY)
-				{
-					Action = FFileChangeData::FCA_Modified;
-				}
-				else if ((Event->mask & IN_DELETE) || (Event->mask & IN_MOVED_FROM))
-				{
-					Action = FFileChangeData::FCA_Removed;
-				}
-
-				Ptr += sizeof(inotify_event)+Event->len;
-
-				// find the file
-				FString Filename;
-				for (int32 i = 0; i < AllFiles.Num(); ++i)
-				{
-					if (WatchDescriptor[i] == Event->wd)
-					{
-						Filename = AllFiles[i];
-						break;
-					}
-				}
-
-				new (FileChanges)FFileChangeData(Filename, Action);
+				Action = FFileChangeData::FCA_Removed;
 			}
+
+			Ptr += sizeof(struct inotify_event) + Event->len;
+
+			const FString Filename = Directory / UTF8_TO_TCHAR(Event->name);
+			new (FileChanges) FFileChangeData(Filename, Action);
 		}
 	}
 }
diff --git a/Engine/Source/Developer/DirectoryWatcher/Private/Linux/DirectoryWatchRequestLinux.h b/Engine/Source/Developer/DirectoryWatcher/Private/Linux/DirectoryWatchRequestLinux.h
--- a/Engine/Source/Developer/DirectoryWatcher/Private/Linux/DirectoryWatchRequestLinux.h
+++ b/Engine/Source/Developer/DirectoryWatcher/Private/Linux/DirectoryWatchRequestLinux.h
@@ -32,17 +32,19 @@ public:
 
 private:
 
-    bool bRunning;
+	FString Directory;
+
+	bool bRunning;
 	bool bEndWatchRequestInvoked;
 
-    int FileDescriptor;
-    int * WatchDescriptor;
-    int NotifyFilter;
+	int FileDescriptor;
+	int * WatchDescriptor;
+	int NotifyFilter;
 
-    TArray<FString> AllFiles;
+	TArray<FString> AllFiles;
 	TArray<IDirectoryWatcher::FDirectoryChanged> Delegates;
 	TArray<FFileChangeData> FileChanges;
 
-    void Shutdown();
-    void ProcessChanges();
+	void Shutdown();
+	void ProcessChanges();
 };
