Fixes to fixes.

diff --git a/Engine/Source/Developer/DirectoryWatcher/Private/Linux/DirectoryWatchRequestLinux.cpp b/Engine/Source/Developer/DirectoryWatcher/Private/Linux/DirectoryWatchRequestLinux.cpp
--- a/Engine/Source/Developer/DirectoryWatcher/Private/Linux/DirectoryWatchRequestLinux.cpp
+++ b/Engine/Source/Developer/DirectoryWatcher/Private/Linux/DirectoryWatchRequestLinux.cpp
@@ -52,8 +52,8 @@ bool FDirectoryWatchRequestLinux::Init(const FString& InDirectory)
 	IFileManager::Get().FindFilesRecursive(AllFiles, *FullPath, TEXT("*"), false, true);
 
 	// Allocate memory for watch descriptors
-	SIZE_T AllocSize = AllFiles.Num()+1 * sizeof(int));
-	WatchDescriptor = reinterpret_cast<int*>(FMemory::Malloc(AllocSize);
+	SIZE_T AllocSize = AllFiles.Num()+1 * sizeof(int);
+	WatchDescriptor = reinterpret_cast<int*>(FMemory::Malloc(AllocSize));
 	if (WatchDescriptor == nullptr) 
 	{
 		UE_LOG(LogDirectoryWatcher, Error, TEXT("Failed to allocate memory for WatchDescriptor"));
