Disable directory watcher for Linux temporarily.

- Crashes at times, need closer inspection.

diff --git a/Engine/Source/Developer/DirectoryWatcher/Private/Linux/DirectoryWatcherLinux.cpp b/Engine/Source/Developer/DirectoryWatcher/Private/Linux/DirectoryWatcherLinux.cpp
--- a/Engine/Source/Developer/DirectoryWatcher/Private/Linux/DirectoryWatcherLinux.cpp
+++ b/Engine/Source/Developer/DirectoryWatcher/Private/Linux/DirectoryWatcherLinux.cpp
@@ -2,6 +2,8 @@
 
 #include "DirectoryWatcherPrivatePCH.h"
 
+#define WITH_DIRECTORY_WATCHER		0
+
 FDirectoryWatcherLinux::FDirectoryWatcherLinux()
 {
 	NumRequests = 0;
@@ -9,6 +11,7 @@ FDirectoryWatcherLinux::FDirectoryWatcherLinux()
 
 FDirectoryWatcherLinux::~FDirectoryWatcherLinux()
 {
+#if WITH_DIRECTORY_WATCHER
 	if (RequestMap.Num() != 0)
 	{
 		// Delete any remaining requests here. These requests are likely from modules which are still loaded at the time that this module unloads.
@@ -32,13 +35,16 @@ FDirectoryWatcherLinux::~FDirectoryWatcherLinux()
 			NumRequests--;
 		}
 	}
-
+#endif // WITH_DIRECTORY_WATCHER
+	
 	// Make sure every request that was created is destroyed
 	ensure(NumRequests == 0);
 }
 
 bool FDirectoryWatcherLinux::RegisterDirectoryChangedCallback( const FString& Directory, const FDirectoryChanged& InDelegate )
 {
+#if WITH_DIRECTORY_WATCHER
+	
 	FDirectoryWatchRequestLinux** RequestPtr = RequestMap.Find(Directory);
 	FDirectoryWatchRequestLinux* Request = NULL;
 
@@ -69,10 +75,17 @@ bool FDirectoryWatcherLinux::RegisterDirectoryChangedCallback( const FString& Di
 	Request->AddDelegate(InDelegate);
 
 	return true;
+
+#else
+	
+	return false;
+	
+#endif // WITH_DIRECTORY_WATCHER
 }
 
 bool FDirectoryWatcherLinux::UnregisterDirectoryChangedCallback(const FString& Directory, const FDirectoryChanged& InDelegate)
 {
+#if WITH_DIRECTORY_WATCHER
 	FDirectoryWatchRequestLinux** RequestPtr = RequestMap.Find(Directory);
 
 	if (RequestPtr)
@@ -98,11 +111,15 @@ bool FDirectoryWatcherLinux::UnregisterDirectoryChangedCallback(const FString& D
 		}
 
 	}
+#endif // WITH_DIRECTORY_WATCHER
+
 	return false;
 }
 
 void FDirectoryWatcherLinux::Tick(float DeltaSeconds)
 {
+#if WITH_DIRECTORY_WATCHER
+
 	// Delete unregistered requests
 	for (int32 RequestIdx = RequestsPendingDelete.Num() - 1; RequestIdx >= 0; --RequestIdx)
 	{
@@ -117,4 +134,6 @@ void FDirectoryWatcherLinux::Tick(float DeltaSeconds)
 	{
 		RequestIt.Value()->ProcessPendingNotifications();
 	}
+
+#endif // WITH_DIRECTORY_WATCHER
 }
