Re-enable DirectoryWatcher on Linux.

- Experimental; caused crashes in the past.

diff --git a/Engine/Source/Developer/DirectoryWatcher/Private/Linux/DirectoryWatcherLinux.cpp b/Engine/Source/Developer/DirectoryWatcher/Private/Linux/DirectoryWatcherLinux.cpp
--- a/Engine/Source/Developer/DirectoryWatcher/Private/Linux/DirectoryWatcherLinux.cpp
+++ b/Engine/Source/Developer/DirectoryWatcher/Private/Linux/DirectoryWatcherLinux.cpp
@@ -2,7 +2,7 @@
 
 #include "DirectoryWatcherPrivatePCH.h"
 
-#define WITH_DIRECTORY_WATCHER		0
+#define WITH_DIRECTORY_WATCHER		1
 
 FDirectoryWatcherLinux::FDirectoryWatcherLinux()
 {
@@ -95,7 +95,7 @@ bool FDirectoryWatcherLinux::UnregisterDirectoryChangedCallback(const FString& D
 
 		FDirectoryWatchRequestLinux* Request = *RequestPtr;
 
-		if (Request->RemoveDelegate(InDelegate))
+		if (Request->DEPRECATED_RemoveDelegate(InDelegate))
 		{
 			if (!Request->HasDelegates())
 			{
