Fix possible culprit of memory stomping in Linux directory  watcher.

diff --git a/Engine/Source/Developer/DirectoryWatcher/Private/Linux/DirectoryWatchRequestLinux.cpp b/Engine/Source/Developer/DirectoryWatcher/Private/Linux/DirectoryWatchRequestLinux.cpp
--- a/Engine/Source/Developer/DirectoryWatcher/Private/Linux/DirectoryWatchRequestLinux.cpp
+++ b/Engine/Source/Developer/DirectoryWatcher/Private/Linux/DirectoryWatchRequestLinux.cpp
@@ -52,7 +52,7 @@ bool FDirectoryWatchRequestLinux::Init(const FString& InDirectory)
 	IFileManager::Get().FindFilesRecursive(AllFiles, *FullPath, TEXT("*"), false, true);
 
 	// Allocate memory for watch descriptors
-	SIZE_T AllocSize = AllFiles.Num()+1 * sizeof(int);
+	SIZE_T AllocSize = (AllFiles.Num()+1) * sizeof(int);
 	WatchDescriptor = reinterpret_cast<int*>(FMemory::Malloc(AllocSize));
 	if (WatchDescriptor == nullptr) 
 	{
