Linux: add watch for the directory too, not just subdirs.

- Hopefully fixes DirectoryWatcher on Linux.

diff --git a/Engine/Source/Developer/DirectoryWatcher/Private/Linux/DirectoryWatchRequestLinux.cpp b/Engine/Source/Developer/DirectoryWatcher/Private/Linux/DirectoryWatchRequestLinux.cpp
--- a/Engine/Source/Developer/DirectoryWatcher/Private/Linux/DirectoryWatchRequestLinux.cpp
+++ b/Engine/Source/Developer/DirectoryWatcher/Private/Linux/DirectoryWatchRequestLinux.cpp
@@ -39,6 +39,7 @@ bool FDirectoryWatchRequestLinux::Init(const FString& InDirectory)
 
 	// Make sure the path is absolute
 	const FString FullPath = FPaths::ConvertRelativePathToFull(InDirectory);
+	UE_LOG(LogDirectoryWatcher, Verbose, TEXT("Adding watch for directory tree '%s'"), *FullPath);
 
 	FileDescriptor = inotify_init1(IN_NONBLOCK | IN_CLOEXEC);
 
@@ -49,10 +50,13 @@ bool FDirectoryWatchRequestLinux::Init(const FString& InDirectory)
 		return false;
 	}
 
+	// find all subdirs
 	IFileManager::Get().FindFilesRecursive(AllFiles, *FullPath, TEXT("*"), false, true);
+	// add the path as well
+	AllFiles.Add(FullPath);
 
 	// Allocate memory for watch descriptors
-	SIZE_T AllocSize = (AllFiles.Num()+1) * sizeof(int);
+	SIZE_T AllocSize = AllFiles.Num() * sizeof(int);
 	WatchDescriptor = reinterpret_cast<int*>(FMemory::Malloc(AllocSize));
 	if (WatchDescriptor == nullptr) 
 	{
@@ -63,13 +67,14 @@ bool FDirectoryWatchRequestLinux::Init(const FString& InDirectory)
 
 	for (int32 FileIdx = 0; FileIdx < AllFiles.Num(); ++FileIdx)
 	{
-			const FString& FolderName = AllFiles[FileIdx];
-			WatchDescriptor[FileIdx] = inotify_add_watch(FileDescriptor, TCHAR_TO_UTF8(*FolderName), NotifyFilter); //FileIdx+1
-			if (WatchDescriptor[FileIdx] == -1) 
-			{
-				UE_LOG(LogDirectoryWatcher, Error, TEXT("inotify_add_watch cannot watch folder %s"), *FolderName);
-				return false;
-			}
+		const FString& FolderName = AllFiles[FileIdx];
+
+		WatchDescriptor[FileIdx] = inotify_add_watch(FileDescriptor, TCHAR_TO_UTF8(*FolderName), NotifyFilter); //FileIdx+1
+		if (WatchDescriptor[FileIdx] == -1) 
+		{
+			UE_LOG(LogDirectoryWatcher, Error, TEXT("inotify_add_watch cannot watch folder %s"), *FolderName);
+			return false;
+		}
 	}
 
 	bRunning = true;
