Improvements to mac directory watcher.
Made some changes to report file changes more accurately. Automation tests to follow.

diff --git a/Engine/Source/Developer/DirectoryWatcher/Private/Mac/DirectoryWatchRequestMac.cpp b/Engine/Source/Developer/DirectoryWatcher/Private/Mac/DirectoryWatchRequestMac.cpp
--- a/Engine/Source/Developer/DirectoryWatcher/Private/Mac/DirectoryWatchRequestMac.cpp
+++ b/Engine/Source/Developer/DirectoryWatcher/Private/Mac/DirectoryWatchRequestMac.cpp
@@ -61,7 +61,7 @@ bool FDirectoryWatchRequestMac::Init(const FString& InDirectory)
 	CFStringRef FullPathMac = FPlatformString::TCHARToCFString(*FullPath);
 	CFArrayRef PathsToWatch = CFArrayCreate(NULL, (const void**)&FullPathMac, 1, NULL);
 
-	CFAbsoluteTime Latency = 2.0;	// seconds
+	CFAbsoluteTime Latency = 0.2;	// seconds
 
 	FSEventStreamContext Context;
 	Context.version = 0;
@@ -170,30 +170,35 @@ void FDirectoryWatchRequestMac::ProcessChanges( size_t EventCount, void* EventPa
 		FFileChangeData::EFileChangeAction Action;
 
 		const bool bFileAdded = ( Flags & kFSEventStreamEventFlagItemCreated );
-		const bool bFileModified = ( Flags & ( kFSEventStreamEventFlagItemRenamed | kFSEventStreamEventFlagItemModified ) );
+		const bool bFileRenamed = ( Flags & kFSEventStreamEventFlagItemRenamed );
+		const bool bFileModified = ( Flags & kFSEventStreamEventFlagItemModified );
 		const bool bFileRemoved = ( Flags & kFSEventStreamEventFlagItemRemoved );
+
 		bool bFileNeedsChecking = false;
 
-		if( bFileAdded )
+		// File modifications take precendent over everything, unless the file has actually been deleted. We only handle newly created files when 
+		// kFSEventStreamEventFlagItemCreated is exclusively set. This flag can often be set when files have been renamed or copied over the top,
+		// but we abstract this behavior and just mark it as modified.
+		if ( bFileModified )
 		{
-			if( bFileRemoved )
-			{
-				// Event indicates that file was both created and removed. To find out which of those events
-				// happened later, we have to check for file's presence.
-				bFileNeedsChecking = true;
-			}
+			bFileNeedsChecking = true;
+			Action = FFileChangeData::FCA_Modified;
+		}
+		else if ( bFileRenamed )
+		{
+			// for now, renames are abstracted as deletes/adds
+			bFileNeedsChecking = true;
 			Action = FFileChangeData::FCA_Added;
 		}
-		else if( bFileModified )
+		else if( bFileAdded )
 		{
-			// File modified
 			if( bFileRemoved )
 			{
-				// Event indicates that file was both modified and removed. To find out which of those events
+				// Event indicates that file was both created and removed. To find out which of those events
 				// happened later, we have to check for file's presence.
 				bFileNeedsChecking = true;
 			}
-			Action = FFileChangeData::FCA_Modified;
+			Action = FFileChangeData::FCA_Added;
 		}
 		else if( bFileRemoved )
 		{
