Automation tests for directory watchers

diff --git a/Engine/Source/Developer/DirectoryWatcher/Private/DirectoryWatcherTests.cpp b/Engine/Source/Developer/DirectoryWatcher/Private/DirectoryWatcherTests.cpp
new file mode 100644
--- /dev/null
+++ b/Engine/Source/Developer/DirectoryWatcher/Private/DirectoryWatcherTests.cpp
@@ -0,0 +1,272 @@
+// Copyright 1998-2015 Epic Games, Inc. All Rights Reserved.
+
+#include "DirectoryWatcherPrivatePCH.h"
+#include "DirectoryWatcherModule.h"
+
+DEFINE_LOG_CATEGORY_STATIC(LogDirectoryWatcherTests, Log, All);
+
+struct FDirectoryWatcherTestPayload
+{
+	FDelegateHandle WatcherDelegate;
+	FString WorkingDir;
+	TMap<FString, FFileChangeData::EFileChangeAction> ReportedChanges;
+
+	FDirectoryWatcherTestPayload(const FString& InWorkingDir)
+		: WorkingDir(InWorkingDir)
+	{
+		IFileManager::Get().MakeDirectory(*WorkingDir, true);
+
+		FDirectoryWatcherModule& Module = FModuleManager::LoadModuleChecked<FDirectoryWatcherModule>(TEXT("DirectoryWatcher"));
+		if (IDirectoryWatcher* DirectoryWatcher = Module.Get())
+		{
+			auto Callback = IDirectoryWatcher::FDirectoryChanged::CreateRaw(this, &FDirectoryWatcherTestPayload::OnDirectoryChanged);
+			DirectoryWatcher->RegisterDirectoryChangedCallback_Handle(WorkingDir, Callback, WatcherDelegate);
+		}
+	}
+	~FDirectoryWatcherTestPayload()
+	{
+		IFileManager::Get().DeleteDirectory(*WorkingDir, false, true);
+
+		FDirectoryWatcherModule& Module = FModuleManager::LoadModuleChecked<FDirectoryWatcherModule>(TEXT("DirectoryWatcher"));
+		if (IDirectoryWatcher* DirectoryWatcher = Module.Get())
+		{
+			DirectoryWatcher->UnregisterDirectoryChangedCallback_Handle(WorkingDir, WatcherDelegate);
+		}
+	}
+
+	void OnDirectoryChanged(const TArray<FFileChangeData>& InFileChanges)
+	{
+		for (const auto& Change : InFileChanges)
+		{
+			FString RelativeFilename = *FPaths::ConvertRelativePathToFull(Change.Filename) + WorkingDir.Len();
+
+			UE_LOG(LogDirectoryWatcherTests, Log, TEXT("File '%s'. Code: %d."), *Change.Filename, (uint8)Change.Action);
+
+			FFileChangeData::EFileChangeAction* Existing = ReportedChanges.Find(RelativeFilename);
+			if (Existing)
+			{
+				switch (Change.Action)
+				{
+				case FFileChangeData::FCA_Added:
+					*Existing = FFileChangeData::FCA_Modified;
+					break;
+
+				case FFileChangeData::FCA_Modified:
+					// We ignore these since added + modified == added, and removed + modified = removed.
+					break;
+
+				case FFileChangeData::FCA_Removed:
+					*Existing = FFileChangeData::FCA_Removed;
+					break;
+				}
+			}
+			else
+			{
+				ReportedChanges.Add(RelativeFilename, Change.Action);
+			}
+		}
+	}
+};
+
+class FDelayedCallbackLatentCommand : public IAutomationLatentCommand
+{
+public:
+	FDelayedCallbackLatentCommand(TFunction<void()> InCallback, float InDelay = 0.1f)
+		: Callback(InCallback), Delay(InDelay)
+	{}
+
+	virtual bool Update() override
+	{
+		float NewTime = FPlatformTime::Seconds();
+		if (NewTime - StartTime >= Delay)
+		{
+			Callback();
+			return true;
+		}
+		return false;
+	}
+
+private:
+	TFunction<void()> Callback;
+	float Delay;
+};
+
+FString GetWorkingDir()
+{
+	return FPaths::ConvertRelativePathToFull(FPaths::AutomationTransientDir() / TEXT("DirectoryWatcher")) / TEXT("");
+}
+
+static const float TestTickDelay = 0.1f;
+IMPLEMENT_SIMPLE_AUTOMATION_TEST(FDirectoryWatcherSimpleCreateTest, "Plugins.Directory Watcher.Simple Create", EAutomationTestFlags::ATF_None)
+bool FDirectoryWatcherSimpleCreateTest::RunTest(const FString& Parameters)
+{
+	const FString WorkingDir = GetWorkingDir();
+
+	static const TCHAR* Filename = TEXT("created.tmp");
+
+	// Start watching the directory
+	TSharedPtr<FDirectoryWatcherTestPayload> Test = MakeShareable(new FDirectoryWatcherTestPayload(WorkingDir));
+
+	// Give the stream time to start up before doing the test
+	ADD_LATENT_AUTOMATION_COMMAND(FDelayedCallbackLatentCommand([=]{
+
+		// Create a file and check that it got reported as created
+		FFileHelper::SaveStringToFile(TEXT(""), *(WorkingDir / Filename));
+
+		ADD_LATENT_AUTOMATION_COMMAND(FDelayedCallbackLatentCommand([=]{
+
+			FFileChangeData::EFileChangeAction* Action = Test->ReportedChanges.Find(Filename);
+
+			if (!Action || *Action != FFileChangeData::FCA_Added)
+			{
+				UE_LOG(LogDirectoryWatcherTests, Error, TEXT("New file '%s' was not correctly reported as being added."), Filename);
+			}
+
+		}));
+
+	}));
+
+	return true;
+}
+
+IMPLEMENT_SIMPLE_AUTOMATION_TEST(FDirectoryWatcherSimpleModifyTest, "Plugins.Directory Watcher.Simple Modify", EAutomationTestFlags::ATF_None)
+bool FDirectoryWatcherSimpleModifyTest::RunTest(const FString& Parameters)
+{
+	const FString WorkingDir = GetWorkingDir();
+
+	static const TCHAR* Filename = TEXT("modified.tmp");
+
+	// Create a file first
+	FFileHelper::SaveStringToFile(TEXT(""), *(WorkingDir / Filename));
+
+	ADD_LATENT_AUTOMATION_COMMAND(FDelayedCallbackLatentCommand([=]{
+
+		// Start watching the directory
+		TSharedPtr<FDirectoryWatcherTestPayload> Test = MakeShareable(new FDirectoryWatcherTestPayload(WorkingDir));
+
+		ADD_LATENT_AUTOMATION_COMMAND(FDelayedCallbackLatentCommand([=]{
+
+			// Manipulate the file
+			FFileHelper::SaveStringToFile(TEXT("Some content"), *(WorkingDir / Filename));
+
+			ADD_LATENT_AUTOMATION_COMMAND(FDelayedCallbackLatentCommand([=]{
+
+				FFileChangeData::EFileChangeAction* Action = Test->ReportedChanges.Find(Filename);
+
+				if (!Action || *Action != FFileChangeData::FCA_Modified)
+				{
+					UE_LOG(LogDirectoryWatcherTests, Error, TEXT("File '%s' was not correctly reported as being modified."), Filename);
+				}
+
+			}));
+
+		}));
+
+	}));
+
+	return true;
+}
+
+IMPLEMENT_SIMPLE_AUTOMATION_TEST(FDirectoryWatcherSimpleDeleteTest, "Plugins.Directory Watcher.Simple Delete", EAutomationTestFlags::ATF_None)
+bool FDirectoryWatcherSimpleDeleteTest::RunTest(const FString& Parameters)
+{
+	const FString WorkingDir = GetWorkingDir();
+
+	static const TCHAR* Filename = TEXT("removed.tmp");
+
+	// Create a file first
+	FFileHelper::SaveStringToFile(TEXT(""), *(WorkingDir / Filename));
+
+	// Start watching the directory
+	TSharedPtr<FDirectoryWatcherTestPayload> Test = MakeShareable(new FDirectoryWatcherTestPayload(WorkingDir));
+
+	// Give the stream time to start up before doing the test
+	ADD_LATENT_AUTOMATION_COMMAND(FDelayedCallbackLatentCommand([=]{
+
+		// Delete the file
+		IFileManager::Get().Delete(*(WorkingDir / Filename));
+
+		ADD_LATENT_AUTOMATION_COMMAND(FDelayedCallbackLatentCommand([=]{
+
+			FFileChangeData::EFileChangeAction* Action = Test->ReportedChanges.Find(Filename);
+
+			if (!Action || *Action != FFileChangeData::FCA_Removed)
+			{
+				UE_LOG(LogDirectoryWatcherTests, Error, TEXT("File '%s' was not correctly reported as being deleted."), Filename);
+			}
+
+		}));
+
+	}));
+
+	return true;
+}
+
+IMPLEMENT_SIMPLE_AUTOMATION_TEST(FDirectoryWatcherSubFolderTest, "Plugins.Directory Watcher.Sub Folder", EAutomationTestFlags::ATF_None)
+bool FDirectoryWatcherSubFolderTest::RunTest(const FString& Parameters)
+{
+	const FString WorkingDir = GetWorkingDir();
+
+	static const TCHAR* CreatedFilename = TEXT("sub_folder/created.tmp");
+	static const TCHAR* ModifiedFilename = TEXT("sub_folder/modified.tmp");
+	static const TCHAR* RemovedFilename = TEXT("sub_folder/removed.tmp");
+
+	// Create the file we wish to modify/delete first
+	FFileHelper::SaveStringToFile(TEXT(""), *(WorkingDir / ModifiedFilename));
+	FFileHelper::SaveStringToFile(TEXT(""), *(WorkingDir / RemovedFilename));
+
+	// Give the stream time to start up before doing the test
+	ADD_LATENT_AUTOMATION_COMMAND(FDelayedCallbackLatentCommand([=]{
+
+		// Start watching the directory
+		TSharedPtr<FDirectoryWatcherTestPayload> Test = MakeShareable(new FDirectoryWatcherTestPayload(WorkingDir));
+
+		// Give the stream time to start up before doing the test
+		ADD_LATENT_AUTOMATION_COMMAND(FDelayedCallbackLatentCommand([=]{
+
+			// Create a new file
+			FFileHelper::SaveStringToFile(TEXT(""), *(WorkingDir / CreatedFilename));
+			// Modify another file
+			FFileHelper::SaveStringToFile(TEXT("Some content"), *(WorkingDir / ModifiedFilename));
+			// Delete a file
+			IFileManager::Get().Delete(*(WorkingDir / RemovedFilename));
+
+			ADD_LATENT_AUTOMATION_COMMAND(FDelayedCallbackLatentCommand([=]{
+
+				FFileChangeData::EFileChangeAction* Action = Test->ReportedChanges.Find(CreatedFilename);
+				if (!Action || *Action != FFileChangeData::FCA_Added)
+				{
+					UE_LOG(LogDirectoryWatcherTests, Error, TEXT("File '%s' was not correctly reported as being added."), CreatedFilename);
+				}
+
+				Action = Test->ReportedChanges.Find(ModifiedFilename);
+				if (!Action || *Action != FFileChangeData::FCA_Modified)
+				{
+					UE_LOG(LogDirectoryWatcherTests, Error, TEXT("File '%s' was not correctly reported as being modified."), ModifiedFilename);
+				}
+
+				Action = Test->ReportedChanges.Find(RemovedFilename);
+				if (!Action || *Action != FFileChangeData::FCA_Removed)
+				{
+					UE_LOG(LogDirectoryWatcherTests, Error, TEXT("File '%s' was not correctly reported as being deleted."), RemovedFilename);
+				}
+
+			}));
+
+		}));
+
+	}));
+
+	return true;
+}
+
+
+
+
+
+
+
+
+
+
+
