Nonunity build fix

diff --git a/Engine/Source/Developer/DirectoryWatcher/Private/DirectoryWatcherTests.cpp b/Engine/Source/Developer/DirectoryWatcher/Private/DirectoryWatcherTests.cpp
--- a/Engine/Source/Developer/DirectoryWatcher/Private/DirectoryWatcherTests.cpp
+++ b/Engine/Source/Developer/DirectoryWatcher/Private/DirectoryWatcherTests.cpp
@@ -2,6 +2,7 @@
 
 #include "DirectoryWatcherPrivatePCH.h"
 #include "DirectoryWatcherModule.h"
+#include "ModuleManager.h"
 
 DEFINE_LOG_CATEGORY_STATIC(LogDirectoryWatcherTests, Log, All);
 
