Changing directory watcher tests to only be run in the editor

diff --git a/Engine/Source/Developer/DirectoryWatcher/Private/DirectoryWatcherTests.cpp b/Engine/Source/Developer/DirectoryWatcher/Private/DirectoryWatcherTests.cpp
--- a/Engine/Source/Developer/DirectoryWatcher/Private/DirectoryWatcherTests.cpp
+++ b/Engine/Source/Developer/DirectoryWatcher/Private/DirectoryWatcherTests.cpp
@@ -98,7 +98,7 @@ FString GetWorkingDir()
 }
 
 static const float TestTickDelay = 0.1f;
-IMPLEMENT_SIMPLE_AUTOMATION_TEST(FDirectoryWatcherSimpleCreateTest, "Plugins.Directory Watcher.Simple Create", EAutomationTestFlags::ATF_None)
+IMPLEMENT_SIMPLE_AUTOMATION_TEST(FDirectoryWatcherSimpleCreateTest, "Plugins.Directory Watcher.Simple Create", EAutomationTestFlags::ATF_Editor)
 bool FDirectoryWatcherSimpleCreateTest::RunTest(const FString& Parameters)
 {
 	const FString WorkingDir = GetWorkingDir();
