Changing directory watcher tests to only be run in the editor

diff --git a/Engine/Source/Developer/DirectoryWatcher/Private/DirectoryWatcherTests.cpp b/Engine/Source/Developer/DirectoryWatcher/Private/DirectoryWatcherTests.cpp
--- a/Engine/Source/Developer/DirectoryWatcher/Private/DirectoryWatcherTests.cpp
+++ b/Engine/Source/Developer/DirectoryWatcher/Private/DirectoryWatcherTests.cpp
@@ -130,7 +130,7 @@ bool FDirectoryWatcherSimpleCreateTest::RunTest(const FString& Parameters)
 	return true;
 }
 
-IMPLEMENT_SIMPLE_AUTOMATION_TEST(FDirectoryWatcherSimpleModifyTest, "Plugins.Directory Watcher.Simple Modify", EAutomationTestFlags::ATF_None)
+IMPLEMENT_SIMPLE_AUTOMATION_TEST(FDirectoryWatcherSimpleModifyTest, "Plugins.Directory Watcher.Simple Modify", EAutomationTestFlags::ATF_Editor)
 bool FDirectoryWatcherSimpleModifyTest::RunTest(const FString& Parameters)
 {
 	const FString WorkingDir = GetWorkingDir();
@@ -168,7 +168,7 @@ bool FDirectoryWatcherSimpleModifyTest::RunTest(const FString& Parameters)
 	return true;
 }
 
-IMPLEMENT_SIMPLE_AUTOMATION_TEST(FDirectoryWatcherSimpleDeleteTest, "Plugins.Directory Watcher.Simple Delete", EAutomationTestFlags::ATF_None)
+IMPLEMENT_SIMPLE_AUTOMATION_TEST(FDirectoryWatcherSimpleDeleteTest, "Plugins.Directory Watcher.Simple Delete", EAutomationTestFlags::ATF_Editor)
 bool FDirectoryWatcherSimpleDeleteTest::RunTest(const FString& Parameters)
 {
 	const FString WorkingDir = GetWorkingDir();
@@ -203,7 +203,7 @@ bool FDirectoryWatcherSimpleDeleteTest::RunTest(const FString& Parameters)
 	return true;
 }
 
-IMPLEMENT_SIMPLE_AUTOMATION_TEST(FDirectoryWatcherSubFolderTest, "Plugins.Directory Watcher.Sub Folder", EAutomationTestFlags::ATF_None)
+IMPLEMENT_SIMPLE_AUTOMATION_TEST(FDirectoryWatcherSubFolderTest, "Plugins.Directory Watcher.Sub Folder", EAutomationTestFlags::ATF_Editor)
 bool FDirectoryWatcherSubFolderTest::RunTest(const FString& Parameters)
 {
 	const FString WorkingDir = GetWorkingDir();
