Linux: more info on watch request failures.

diff --git a/Engine/Source/Developer/DirectoryWatcher/Private/Linux/DirectoryWatchRequestLinux.cpp b/Engine/Source/Developer/DirectoryWatcher/Private/Linux/DirectoryWatchRequestLinux.cpp
--- a/Engine/Source/Developer/DirectoryWatcher/Private/Linux/DirectoryWatchRequestLinux.cpp
+++ b/Engine/Source/Developer/DirectoryWatcher/Private/Linux/DirectoryWatchRequestLinux.cpp
@@ -72,7 +72,11 @@ bool FDirectoryWatchRequestLinux::Init(const FString& InDirectory)
 		WatchDescriptor[FileIdx] = inotify_add_watch(FileDescriptor, TCHAR_TO_UTF8(*FolderName), NotifyFilter); //FileIdx+1
 		if (WatchDescriptor[FileIdx] == -1) 
 		{
-			UE_LOG(LogDirectoryWatcher, Error, TEXT("inotify_add_watch cannot watch folder %s"), *FolderName);
+			int ErrNo = errno;
+			UE_LOG(LogDirectoryWatcher, Error, TEXT("inotify_add_watch cannot watch folder %s (errno = %d, %s)"), *FolderName,
+				ErrNo,
+				ANSI_TO_TCHAR(strerror(ErrNo))
+			);
 			return false;
 		}
 	}
