Add support for directory events in directory watcher, and use it to refresh the list of enabled plugins.

diff --git a/Engine/Plugins/Editor/PluginBrowser/Source/PluginBrowser/Private/SPluginBrowser.cpp b/Engine/Plugins/Editor/PluginBrowser/Source/PluginBrowser/Private/SPluginBrowser.cpp
--- a/Engine/Plugins/Editor/PluginBrowser/Source/PluginBrowser/Private/SPluginBrowser.cpp
+++ b/Engine/Plugins/Editor/PluginBrowser/Source/PluginBrowser/Private/SPluginBrowser.cpp
@@ -35,7 +35,7 @@ void SPluginBrowser::Construct( const FArguments& Args )
 	for(const FString& WatchDirectoryName: WatchDirectoryNames)
 	{
 		FDelegateHandle Handle;
-		if(DirectoryWatcherModule.Get()->RegisterDirectoryChangedCallback_Handle(WatchDirectoryName, IDirectoryWatcher::FDirectoryChanged::CreateRaw(this, &SPluginBrowser::OnPluginDirectoryChanged), Handle))
+		if(DirectoryWatcherModule.Get()->RegisterDirectoryChangedCallback_Handle(WatchDirectoryName, IDirectoryWatcher::FDirectoryChanged::CreateRaw(this, &SPluginBrowser::OnPluginDirectoryChanged), Handle, true))
 		{
 			WatchDirectories.Add(WatchDirectoryName, Handle);
 		}
diff --git a/Engine/Source/Developer/DirectoryWatcher/Private/DirectoryWatcherTests.cpp b/Engine/Source/Developer/DirectoryWatcher/Private/DirectoryWatcherTests.cpp
--- a/Engine/Source/Developer/DirectoryWatcher/Private/DirectoryWatcherTests.cpp
+++ b/Engine/Source/Developer/DirectoryWatcher/Private/DirectoryWatcherTests.cpp
@@ -12,7 +12,7 @@ struct FDirectoryWatcherTestPayload
 	FString WorkingDir;
 	TMap<FString, FFileChangeData::EFileChangeAction> ReportedChanges;
 
-	FDirectoryWatcherTestPayload(const FString& InWorkingDir)
+	FDirectoryWatcherTestPayload(const FString& InWorkingDir, bool bIncludeDirectoryEvents = false)
 		: WorkingDir(InWorkingDir)
 	{
 		IFileManager::Get().MakeDirectory(*WorkingDir, true);
@@ -21,7 +21,7 @@ struct FDirectoryWatcherTestPayload
 		if (IDirectoryWatcher* DirectoryWatcher = Module.Get())
 		{
 			auto Callback = IDirectoryWatcher::FDirectoryChanged::CreateRaw(this, &FDirectoryWatcherTestPayload::OnDirectoryChanged);
-			DirectoryWatcher->RegisterDirectoryChangedCallback_Handle(WorkingDir, Callback, WatcherDelegate);
+			DirectoryWatcher->RegisterDirectoryChangedCallback_Handle(WorkingDir, Callback, WatcherDelegate, bIncludeDirectoryEvents);
 		}
 	}
 	~FDirectoryWatcherTestPayload()
@@ -261,6 +261,50 @@ bool FDirectoryWatcherSubFolderTest::RunTest(const FString& Parameters)
 	return true;
 }
 
+IMPLEMENT_SIMPLE_AUTOMATION_TEST(FDirectoryWatcherNewFolderTest, "System.Plugins.Directory Watcher.New Folder", EAutomationTestFlags::ATF_Editor)
+bool FDirectoryWatcherNewFolderTest::RunTest(const FString& Parameters)
+{
+	const FString WorkingDir = GetWorkingDir();
+
+	static const TCHAR* CreatedDirectory = TEXT("created");
+	static const TCHAR* RemovedDirectory = TEXT("removed");
+
+	// Give the stream time to start up before doing the test
+	ADD_LATENT_AUTOMATION_COMMAND(FDelayedCallbackLatentCommand([=]{
+
+		IFileManager::Get().MakeDirectory(*(WorkingDir / RemovedDirectory), true);
+
+		// Start watching the directory
+		TSharedPtr<FDirectoryWatcherTestPayload> Test = MakeShareable(new FDirectoryWatcherTestPayload(WorkingDir, true));
+
+		// Give the stream time to start up before doing the test
+		ADD_LATENT_AUTOMATION_COMMAND(FDelayedCallbackLatentCommand([=]{
+
+			IFileManager::Get().MakeDirectory(*(WorkingDir / CreatedDirectory), true);
+			IFileManager::Get().DeleteDirectory(*(WorkingDir / RemovedDirectory), true);
+
+			ADD_LATENT_AUTOMATION_COMMAND(FDelayedCallbackLatentCommand([=]{
+
+				FFileChangeData::EFileChangeAction* Action = Test->ReportedChanges.Find(CreatedDirectory);
+				if (!Action || *Action != FFileChangeData::FCA_Added)
+				{
+					UE_LOG(LogDirectoryWatcherTests, Error, TEXT("Folder '%s' was not correctly reported as being added."), CreatedDirectory);
+				}
+
+				Action = Test->ReportedChanges.Find(RemovedDirectory);
+				if (!Action || *Action != FFileChangeData::FCA_Removed)
+				{
+					UE_LOG(LogDirectoryWatcherTests, Error, TEXT("Folder '%s' was not correctly reported as being removed."), RemovedDirectory);
+				}
+
+			}));
+
+		}));
+
+	}));
+
+	return true;
+}
 
 
 
diff --git a/Engine/Source/Developer/DirectoryWatcher/Private/Mac/DirectoryWatchRequestMac.cpp b/Engine/Source/Developer/DirectoryWatcher/Private/Mac/DirectoryWatchRequestMac.cpp
--- a/Engine/Source/Developer/DirectoryWatcher/Private/Mac/DirectoryWatchRequestMac.cpp
+++ b/Engine/Source/Developer/DirectoryWatcher/Private/Mac/DirectoryWatchRequestMac.cpp
@@ -13,9 +13,10 @@ void DirectoryWatchMacCallback( ConstFSEventStreamRef StreamRef, void* WatchRequ
 
 // ============================================================================================================================
 
-FDirectoryWatchRequestMac::FDirectoryWatchRequestMac()
+FDirectoryWatchRequestMac::FDirectoryWatchRequestMac(bool bInIncludeDirectoryEvents)
 :	bRunning(false)
 ,	bEndWatchRequestInvoked(false)
+,	bIncludeDirectoryEvents(bInIncludeDirectoryEvents)
 {
 }
 
@@ -161,8 +162,11 @@ void FDirectoryWatchRequestMac::ProcessChanges( size_t EventCount, void* EventPa
 		const FSEventStreamEventFlags Flags = EventFlags[EventIndex];
 		if( !(Flags & kFSEventStreamEventFlagItemIsFile) )
 		{
-			// events about directories and symlinks don't concern us
-			continue;
+			if( !bIncludeDirectoryEvents || !(Flags & kFSEventStreamEventFlagItemIsDir))
+			{
+				// events about directories and symlinks don't concern us
+				continue;
+			}
 		}
 
 		// Warning: some events have more than one of created, removed nd modified flag.
diff --git a/Engine/Source/Developer/DirectoryWatcher/Private/Mac/DirectoryWatchRequestMac.h b/Engine/Source/Developer/DirectoryWatcher/Private/Mac/DirectoryWatchRequestMac.h
--- a/Engine/Source/Developer/DirectoryWatcher/Private/Mac/DirectoryWatchRequestMac.h
+++ b/Engine/Source/Developer/DirectoryWatcher/Private/Mac/DirectoryWatchRequestMac.h
@@ -8,7 +8,7 @@
 class FDirectoryWatchRequestMac
 {
 public:
-	FDirectoryWatchRequestMac();
+	FDirectoryWatchRequestMac(bool bInIncludeDirectoryEvents);
 	virtual ~FDirectoryWatchRequestMac();
 
 	/** Sets up the directory handle and request information */
@@ -35,6 +35,7 @@ private:
 	FSEventStreamRef	EventStream;
 	bool				bRunning;
 	bool				bEndWatchRequestInvoked;
+	bool				bIncludeDirectoryEvents;
 
 	TArray<IDirectoryWatcher::FDirectoryChanged> Delegates;
 	TArray<FFileChangeData> FileChanges;
diff --git a/Engine/Source/Developer/DirectoryWatcher/Private/Mac/DirectoryWatcherMac.cpp b/Engine/Source/Developer/DirectoryWatcher/Private/Mac/DirectoryWatcherMac.cpp
--- a/Engine/Source/Developer/DirectoryWatcher/Private/Mac/DirectoryWatcherMac.cpp
+++ b/Engine/Source/Developer/DirectoryWatcher/Private/Mac/DirectoryWatcherMac.cpp
@@ -51,7 +51,7 @@ bool FDirectoryWatcherMac::RegisterDirectoryChangedCallback( const FString& Dire
 	}
 	else
 	{
-		Request = new FDirectoryWatchRequestMac();
+		Request = new FDirectoryWatchRequestMac(false);
 		NumRequests++;
 
 		// Begin reading directory changes
@@ -102,7 +102,7 @@ bool FDirectoryWatcherMac::UnregisterDirectoryChangedCallback( const FString& Di
 	return false;
 }
 
-bool FDirectoryWatcherMac::RegisterDirectoryChangedCallback_Handle( const FString& Directory, const FDirectoryChanged& InDelegate, FDelegateHandle& OutHandle )
+bool FDirectoryWatcherMac::RegisterDirectoryChangedCallback_Handle( const FString& Directory, const FDirectoryChanged& InDelegate, FDelegateHandle& OutHandle, bool bIncludeDirectoryChanges )
 {
 	FDirectoryWatchRequestMac** RequestPtr = RequestMap.Find(Directory);
 	FDirectoryWatchRequestMac* Request = NULL;
@@ -116,7 +116,7 @@ bool FDirectoryWatcherMac::RegisterDirectoryChangedCallback_Handle( const FStrin
 	}
 	else
 	{
-		Request = new FDirectoryWatchRequestMac();
+		Request = new FDirectoryWatchRequestMac(bIncludeDirectoryChanges);
 		NumRequests++;
 
 		// Begin reading directory changes
diff --git a/Engine/Source/Developer/DirectoryWatcher/Private/Mac/DirectoryWatcherMac.h b/Engine/Source/Developer/DirectoryWatcher/Private/Mac/DirectoryWatcherMac.h
--- a/Engine/Source/Developer/DirectoryWatcher/Private/Mac/DirectoryWatcherMac.h
+++ b/Engine/Source/Developer/DirectoryWatcher/Private/Mac/DirectoryWatcherMac.h
@@ -11,7 +11,7 @@ public:
 
 	virtual bool RegisterDirectoryChangedCallback (const FString& Directory, const FDirectoryChanged& InDelegate) override;
 	virtual bool UnregisterDirectoryChangedCallback (const FString& Directory, const FDirectoryChanged& InDelegate) override;
-	virtual bool RegisterDirectoryChangedCallback_Handle (const FString& Directory, const FDirectoryChanged& InDelegate, FDelegateHandle& OutHandle) override;
+	virtual bool RegisterDirectoryChangedCallback_Handle (const FString& Directory, const FDirectoryChanged& InDelegate, FDelegateHandle& OutHandle, bool bIncludeDirectoryChanges) override;
 	virtual bool UnregisterDirectoryChangedCallback_Handle (const FString& Directory, FDelegateHandle InHandle) override;
 	virtual void Tick (float DeltaSeconds) override;
 
diff --git a/Engine/Source/Developer/DirectoryWatcher/Private/Windows/DirectoryWatchRequestWindows.cpp b/Engine/Source/Developer/DirectoryWatcher/Private/Windows/DirectoryWatchRequestWindows.cpp
--- a/Engine/Source/Developer/DirectoryWatcher/Private/Windows/DirectoryWatchRequestWindows.cpp
+++ b/Engine/Source/Developer/DirectoryWatcher/Private/Windows/DirectoryWatchRequestWindows.cpp
@@ -2,7 +2,7 @@
 
 #include "DirectoryWatcherPrivatePCH.h"
 
-FDirectoryWatchRequestWindows::FDirectoryWatchRequestWindows()
+FDirectoryWatchRequestWindows::FDirectoryWatchRequestWindows(bool bIncludeDirectoryEvents)
 {
 	bPendingDelete = false;
 	bEndWatchRequestInvoked = false;
@@ -10,7 +10,7 @@ FDirectoryWatchRequestWindows::FDirectoryWatchRequestWindows()
 	MaxChanges = 16384;
 	bWatchSubtree = true;
 
-	NotifyFilter = FILE_NOTIFY_CHANGE_FILE_NAME | FILE_NOTIFY_CHANGE_LAST_WRITE | FILE_NOTIFY_CHANGE_CREATION;
+	NotifyFilter = FILE_NOTIFY_CHANGE_FILE_NAME | (bIncludeDirectoryEvents? FILE_NOTIFY_CHANGE_DIR_NAME : 0) | FILE_NOTIFY_CHANGE_LAST_WRITE | FILE_NOTIFY_CHANGE_CREATION;
 
 	DirectoryHandle = INVALID_HANDLE_VALUE;
 
diff --git a/Engine/Source/Developer/DirectoryWatcher/Private/Windows/DirectoryWatchRequestWindows.h b/Engine/Source/Developer/DirectoryWatcher/Private/Windows/DirectoryWatchRequestWindows.h
--- a/Engine/Source/Developer/DirectoryWatcher/Private/Windows/DirectoryWatchRequestWindows.h
+++ b/Engine/Source/Developer/DirectoryWatcher/Private/Windows/DirectoryWatchRequestWindows.h
@@ -5,7 +5,7 @@
 class FDirectoryWatchRequestWindows
 {
 public:
-	FDirectoryWatchRequestWindows();
+	FDirectoryWatchRequestWindows(bool bIncludeDirectoryEvents);
 	virtual ~FDirectoryWatchRequestWindows();
 
 	/** Sets up the directory handle and request information */
diff --git a/Engine/Source/Developer/DirectoryWatcher/Private/Windows/DirectoryWatcherWindows.cpp b/Engine/Source/Developer/DirectoryWatcher/Private/Windows/DirectoryWatcherWindows.cpp
--- a/Engine/Source/Developer/DirectoryWatcher/Private/Windows/DirectoryWatcherWindows.cpp
+++ b/Engine/Source/Developer/DirectoryWatcher/Private/Windows/DirectoryWatcherWindows.cpp
@@ -53,7 +53,7 @@ bool FDirectoryWatcherWindows::RegisterDirectoryChangedCallback( const FString&
 	}
 	else
 	{
-		Request = new FDirectoryWatchRequestWindows();
+		Request = new FDirectoryWatchRequestWindows(false);
 		NumRequests++;
 
 		// Begin reading directory changes
@@ -105,7 +105,7 @@ bool FDirectoryWatcherWindows::UnregisterDirectoryChangedCallback( const FString
 	return false;
 }
 
-bool FDirectoryWatcherWindows::RegisterDirectoryChangedCallback_Handle( const FString& Directory, const FDirectoryChanged& InDelegate, FDelegateHandle& Handle )
+bool FDirectoryWatcherWindows::RegisterDirectoryChangedCallback_Handle( const FString& Directory, const FDirectoryChanged& InDelegate, FDelegateHandle& Handle, bool bIncludeDirectoryChanges )
 {
 	FDirectoryWatchRequestWindows** RequestPtr = RequestMap.Find(Directory);
 	FDirectoryWatchRequestWindows* Request = NULL;
@@ -119,7 +119,7 @@ bool FDirectoryWatcherWindows::RegisterDirectoryChangedCallback_Handle( const FS
 	}
 	else
 	{
-		Request = new FDirectoryWatchRequestWindows();
+		Request = new FDirectoryWatchRequestWindows(bIncludeDirectoryChanges);
 		NumRequests++;
 
 		// Begin reading directory changes
diff --git a/Engine/Source/Developer/DirectoryWatcher/Private/Windows/DirectoryWatcherWindows.h b/Engine/Source/Developer/DirectoryWatcher/Private/Windows/DirectoryWatcherWindows.h
--- a/Engine/Source/Developer/DirectoryWatcher/Private/Windows/DirectoryWatcherWindows.h
+++ b/Engine/Source/Developer/DirectoryWatcher/Private/Windows/DirectoryWatcherWindows.h
@@ -17,7 +17,7 @@ public:
 
 	virtual bool RegisterDirectoryChangedCallback (const FString& Directory, const FDirectoryChanged& InDelegate) override;
 	virtual bool UnregisterDirectoryChangedCallback (const FString& Directory, const FDirectoryChanged& InDelegate) override;
-	virtual bool RegisterDirectoryChangedCallback_Handle (const FString& Directory, const FDirectoryChanged& InDelegate, FDelegateHandle& OutHandle) override;
+	virtual bool RegisterDirectoryChangedCallback_Handle (const FString& Directory, const FDirectoryChanged& InDelegate, FDelegateHandle& OutHandle, bool bIncludeDirectoryEvents) override;
 	virtual bool UnregisterDirectoryChangedCallback_Handle (const FString& Directory, FDelegateHandle InHandle) override;
 	virtual void Tick (float DeltaSeconds) override;
 
diff --git a/Engine/Source/Developer/DirectoryWatcher/Public/IDirectoryWatcher.h b/Engine/Source/Developer/DirectoryWatcher/Public/IDirectoryWatcher.h
--- a/Engine/Source/Developer/DirectoryWatcher/Public/IDirectoryWatcher.h
+++ b/Engine/Source/Developer/DirectoryWatcher/Public/IDirectoryWatcher.h
@@ -58,8 +58,9 @@ public:
 	 * @param	Directory to watch
 	 * @param	Delegate to add to our callback list
 	 * @param	The handle to the registered delegate, if the registration was successful.
+	 * @param   Whether to include notifications for changes to actual directories (such as directories being created or removed)
 	 */
-	virtual bool RegisterDirectoryChangedCallback_Handle (const FString& Directory, const FDirectoryChanged& InDelegate, FDelegateHandle& OutHandle) = 0;
+	virtual bool RegisterDirectoryChangedCallback_Handle (const FString& Directory, const FDirectoryChanged& InDelegate, FDelegateHandle& OutHandle, bool bIncludeDirectoryChanges = false) = 0;
 
 	/**
 	 * Unregisters a callback to fire when directories are changed
