Fix DirectoryWatcher on Linux.

diff --git a/Engine/Source/Developer/DirectoryWatcher/Private/Linux/DirectoryWatcherLinux.cpp b/Engine/Source/Developer/DirectoryWatcher/Private/Linux/DirectoryWatcherLinux.cpp
--- a/Engine/Source/Developer/DirectoryWatcher/Private/Linux/DirectoryWatcherLinux.cpp
+++ b/Engine/Source/Developer/DirectoryWatcher/Private/Linux/DirectoryWatcherLinux.cpp
@@ -116,7 +116,7 @@ bool FDirectoryWatcherLinux::UnregisterDirectoryChangedCallback(const FString& D
 	return false;
 }
 
-bool FDirectoryWatcherLinux::RegisterDirectoryChangedCallback_Handle( const FString& Directory, const FDirectoryChanged& InDelegate, FDelegateHandle& OutHandle )
+bool FDirectoryWatcherLinux::RegisterDirectoryChangedCallback_Handle( const FString& Directory, const FDirectoryChanged& InDelegate, FDelegateHandle& OutHandle, bool bIncludeDirectoryChanges )
 {
 #if WITH_DIRECTORY_WATCHER
 	
diff --git a/Engine/Source/Developer/DirectoryWatcher/Private/Linux/DirectoryWatcherLinux.h b/Engine/Source/Developer/DirectoryWatcher/Private/Linux/DirectoryWatcherLinux.h
--- a/Engine/Source/Developer/DirectoryWatcher/Private/Linux/DirectoryWatcherLinux.h
+++ b/Engine/Source/Developer/DirectoryWatcher/Private/Linux/DirectoryWatcherLinux.h
@@ -10,7 +10,7 @@ public:
 
 	virtual bool RegisterDirectoryChangedCallback(const FString& Directory, const FDirectoryChanged& InDelegate) override;
 	virtual bool UnregisterDirectoryChangedCallback(const FString& Directory, const FDirectoryChanged& InDelegate) override;
-	virtual bool RegisterDirectoryChangedCallback_Handle(const FString& Directory, const FDirectoryChanged& InDelegate, FDelegateHandle& OutHandle) override;
+	virtual bool RegisterDirectoryChangedCallback_Handle(const FString& Directory, const FDirectoryChanged& InDelegate, FDelegateHandle& OutHandle, bool bIncludeDirectoryChanges) override;
 	virtual bool UnregisterDirectoryChangedCallback_Handle(const FString& Directory, FDelegateHandle InHandle) override;
 	virtual void Tick (float DeltaSeconds) override;
 
