Linux: multiple fixes for DirectoryWatcher.

- DW now returns correct paths for changes in subdirectory (UE-11861)
- DW will now automatically watch for changes in the whole tree, updating mapping when directories are added/deleted.
- DW now honors bIncludeDirectoryChanges option.
- Code has been hardened (should handle overflow and other failures).
- This change Incorporates parts from PR #1139 by 3dluvr (UE-15443) .
- Passes TestPAL and automation tests.

diff --git a/Engine/Source/Developer/DirectoryWatcher/Private/Linux/DirectoryWatchRequestLinux.cpp b/Engine/Source/Developer/DirectoryWatcher/Private/Linux/DirectoryWatchRequestLinux.cpp
--- a/Engine/Source/Developer/DirectoryWatcher/Private/Linux/DirectoryWatchRequestLinux.cpp
+++ b/Engine/Source/Developer/DirectoryWatcher/Private/Linux/DirectoryWatchRequestLinux.cpp
@@ -16,11 +16,13 @@ FDirectoryWatchRequestLinux::~FDirectoryWatchRequestLinux()
 
 void FDirectoryWatchRequestLinux::Shutdown()
 {
-	FMemory::Free(WatchDescriptor);
 	close(FileDescriptor);
+
+	WatchDescriptorsToPaths.Empty();
+	PathsToWatchDescriptors.Empty();
 }
 
-bool FDirectoryWatchRequestLinux::Init(const FString& InDirectory)
+bool FDirectoryWatchRequestLinux::Init(const FString& InDirectory, bool bInIncludeDirectoryChanges)
 {
 	if (InDirectory.Len() == 0)
 	{
@@ -29,6 +31,7 @@ bool FDirectoryWatchRequestLinux::Init(const FString& InDirectory)
 	}
 
 	Directory = InDirectory;
+	bIncludeDirectoryChanges = bInIncludeDirectoryChanges;
 
 	if (bRunning)
 	{
@@ -51,35 +54,7 @@ bool FDirectoryWatchRequestLinux::Init(const FString& InDirectory)
 	}
 
 	// find all subdirs
-	IFileManager::Get().FindFilesRecursive(AllFiles, *FullPath, TEXT("*"), false, true);
-	// add the path as well
-	AllFiles.Add(FullPath);
-
-	// Allocate memory for watch descriptors
-	SIZE_T AllocSize = AllFiles.Num() * sizeof(int);
-	WatchDescriptor = reinterpret_cast<int*>(FMemory::Malloc(AllocSize));
-	if (WatchDescriptor == nullptr) 
-	{
-		UE_LOG(LogDirectoryWatcher, Error, TEXT("Failed to allocate memory for WatchDescriptor"));
-		return false;
-	}
-	FMemory::Memzero(WatchDescriptor, AllocSize);
-
-	for (int32 FileIdx = 0; FileIdx < AllFiles.Num(); ++FileIdx)
-	{
-		const FString& FolderName = AllFiles[FileIdx];
-
-		WatchDescriptor[FileIdx] = inotify_add_watch(FileDescriptor, TCHAR_TO_UTF8(*FolderName), NotifyFilter); //FileIdx+1
-		if (WatchDescriptor[FileIdx] == -1) 
-		{
-			int ErrNo = errno;
-			UE_LOG(LogDirectoryWatcher, Error, TEXT("inotify_add_watch cannot watch folder %s (errno = %d, %s)"), *FolderName,
-				ErrNo,
-				ANSI_TO_TCHAR(strerror(ErrNo))
-			);
-			return false;
-		}
-	}
+	WatchDirectoryTree(FullPath);
 
 	bRunning = true;
 
@@ -138,9 +113,78 @@ void FDirectoryWatchRequestLinux::ProcessPendingNotifications()
 	}
 }
 
+void FDirectoryWatchRequestLinux::WatchDirectoryTree(const FString & RootAbsolutePath)
+{
+	UE_LOG(LogDirectoryWatcher, VeryVerbose, TEXT("Watching tree '%s'"), *RootAbsolutePath);
+	TArray<FString> AllFiles;
+	IFileManager::Get().FindFilesRecursive(AllFiles, *RootAbsolutePath, TEXT("*"), false, true);
+	// add the path as well
+	AllFiles.Add(RootAbsolutePath);
+
+	for (int32 FileIdx = 0, NumTotal = AllFiles.Num(); FileIdx < NumTotal; ++FileIdx)
+	{
+		const FString& FolderName = AllFiles[FileIdx];
+		checkf(PathsToWatchDescriptors.Find(FolderName) == nullptr, TEXT("Adding a duplicate watch for directory '%s'"), *FolderName);
+
+		int32 WatchDescriptor = inotify_add_watch(FileDescriptor, TCHAR_TO_UTF8(*FolderName), NotifyFilter);
+		if (WatchDescriptor == -1)
+		{
+			int ErrNo = errno;
+			UE_LOG(LogDirectoryWatcher, Error, TEXT("inotify_add_watch cannot watch folder %s (errno = %d, %s)"), *FolderName,
+				ErrNo,
+				ANSI_TO_TCHAR(strerror(ErrNo))
+				);
+			// proceed further
+		}
+
+		UE_LOG(LogDirectoryWatcher, VeryVerbose, TEXT("+ Added a watch %d for '%s'"), WatchDescriptor, *FolderName);
+		// update the mapping
+		WatchDescriptorsToPaths.Add(WatchDescriptor, FolderName);
+		PathsToWatchDescriptors.Add(FolderName, WatchDescriptor);
+	}
+}
+
+void FDirectoryWatchRequestLinux::UnwatchDirectoryTree(const FString & RootAbsolutePath)
+{
+	UE_LOG(LogDirectoryWatcher, VeryVerbose, TEXT("Unwatching tree '%s'"), *RootAbsolutePath);
+	// remove the watch for the folder and all subfolders
+	// since it is expected that there will be a lot of them, just create a new TMap
+	TMap<FString, int32, FDefaultSetAllocator, FCaseSensitiveLookupKeyFuncs<int32>> NewPathsToWatchDescriptors;
+	for (auto MapIt = PathsToWatchDescriptors.CreateIterator(); MapIt; ++MapIt)
+	{
+		if (!MapIt->Key.StartsWith(RootAbsolutePath, ESearchCase::CaseSensitive))
+		{
+			NewPathsToWatchDescriptors.Add(MapIt->Key, MapIt->Value);
+		}
+		else
+		{
+			UE_LOG(LogDirectoryWatcher, VeryVerbose, TEXT("- Removing a watch %d for '%s'"), MapIt->Value, *MapIt->Key);
+
+			// delete the descriptor
+			int RetVal = inotify_rm_watch(FileDescriptor, MapIt->Value);
+
+			// why check for RootAbsolutePath? Because this function may be called when root path has been deleted, and inotify_rm_watch() will fail
+			// removing a watch on a deleted file... yay for API symmetry. Just "leak" the watch descriptor without the warning
+			if (RetVal == -1 || MapIt->Key != RootAbsolutePath)
+			{
+				int ErrNo = errno;
+				UE_LOG(LogDirectoryWatcher, Error, TEXT("inotify_rm_watch cannot remove descriptor %d for folder '%s' (errno = %d, %s)"),
+					MapIt->Value,
+					*MapIt->Key,
+					ErrNo,
+					ANSI_TO_TCHAR(strerror(ErrNo))
+				);
+			}
+			WatchDescriptorsToPaths.Remove(MapIt->Value);
+		}
+	}
+
+	PathsToWatchDescriptors = NewPathsToWatchDescriptors;
+}
+
 void FDirectoryWatchRequestLinux::ProcessChanges()
 {
-	uint8_t Buf[BUF_LEN] __attribute__ ((aligned(__alignof__(struct inotify_event))));
+	uint8_t Buffer[EVENT_BUF_LEN] __attribute__ ((aligned(__alignof__(struct inotify_event))));
 	const struct inotify_event* Event;
 	ssize_t Len = 0;
 	uint8_t* Ptr = nullptr;
@@ -148,46 +192,112 @@ void FDirectoryWatchRequestLinux::ProcessChanges()
 	// Loop while events can be read from inotify file descriptor
 	for (;;)
 	{
+		// default action
 		FFileChangeData::EFileChangeAction Action = FFileChangeData::FCA_Unknown;
-		// Read some events
-		Len = read(FileDescriptor, Buf, sizeof Buf);
+
+		// Read event stream
+		Len = read(FileDescriptor, Buffer, EVENT_BUF_LEN);
+
+		// If the non-blocking read() found no events to read, then it returns -1 with errno set to EAGAIN.
 		if (Len == -1 && errno != EAGAIN)
 		{
-			UE_LOG(LogDirectoryWatcher, Fatal, TEXT("DirectoryWatchRequest ProcessChanges read error"));
-			// unreachable
+			int ErrNo = errno;
+			UE_LOG(LogDirectoryWatcher, Error, TEXT("FDirectoryWatchRequestLinux::ProcessChanges() read() error getting events for path '%s' (errno = %d, %s)"),
+				*Directory,
+				ErrNo,
+				ANSI_TO_TCHAR(strerror(ErrNo))
+				);
 			break;
 		}
 
-		// If the non-blocking read() found no events to read, then it returns -1 with errno set to EAGAIN. 
 		if (Len <= 0)
 		{
 			break;
 		}
 
 		// Loop over all events in the buffer
-		for (Ptr = Buf; Ptr < Buf + Len;)
+		for (Ptr = Buffer; Ptr < Buffer + Len;)
 		{
 			Event = reinterpret_cast<const struct inotify_event *>(Ptr);
 
-			if ((Event->mask & IN_CREATE) || (Event->mask & IN_MOVED_TO))
-			{
-				Action = FFileChangeData::FCA_Added;
-			}
+			// some events can be ignored to match other implementations and tests
+			bool bIgnoreEvent = false;
 
-			if (Event->mask & IN_MODIFY)
+			// skip if overflowed
+			if (Event->wd != -1 && (Event->mask & IN_Q_OVERFLOW) == 0)
 			{
-				Action = FFileChangeData::FCA_Modified;
+				const FString *EventPathPtr = WatchDescriptorsToPaths.Find(Event->wd);
+
+				UE_LOG(LogDirectoryWatcher, VeryVerbose, TEXT("Event: watch descriptor %d, mask 0x%08x, EventPath: '%s'"),
+					Event->wd, Event->mask, EventPathPtr ? *(*EventPathPtr) : TEXT("nullptr"), ANSI_TO_TCHAR(Event->name));
+
+				// if we're geting multiple events (e.g. DELETE, IGNORED) we could ahve removed descriptor on previous iteration,
+				// so we need to handle inability to find it in the map
+				if (EventPathPtr)
+				{
+					FString EventPath = *EventPathPtr;	// get a copy since we can mutate WatchDescriptorToPaths
+					FString AffectedFile = EventPath / ANSI_TO_TCHAR(Event->name);	// by default, some events report about the file itself
+
+					if ((Event->mask & IN_CREATE) || (Event->mask & IN_MOVED_TO))
+					{
+						// if a directory was created/moved, watch it
+						if (Event->mask & IN_ISDIR)
+						{
+							WatchDirectoryTree(AffectedFile);
+							// to be in sync with Windows implementation, ignore events about creating directories unless told so
+							bIgnoreEvent = !bIncludeDirectoryChanges;
+						}
+
+						Action = FFileChangeData::FCA_Added;
+					}
+					else if (Event->mask & IN_MODIFY)
+					{
+						// if a directory was modified, we expect to get events from already watched files in it
+						Action = FFileChangeData::FCA_Modified;
+					}
+					// check if the file/directory itself has been deleted (IGNORED can also be sent on delete)
+					else if ((Event->mask & IN_DELETE_SELF) || (Event->mask & IN_IGNORED) || (Event->mask & IN_UNMOUNT))
+					{
+						// if a directory was deleted, we expect to get events from already watched files in it
+						AffectedFile = EventPath;
+						if (Event->mask & IN_ISDIR)
+						{
+							UnwatchDirectoryTree(EventPath);
+							// to be in sync with Windows implementation, ignore events about deleting directories unless told so
+							bIgnoreEvent = !bIncludeDirectoryChanges;
+						}
+						else
+						{
+							// remove itself from both mappings
+							// NOTE: inotify_rm_watch() will fail here as watch descriptor is no longer valid (deficiency of inotify?).
+							// Just avoid calling inotify_rm_watch() even if it feels like a leak
+							WatchDescriptorsToPaths.Remove(Event->wd);
+							PathsToWatchDescriptors.Remove(EventPath);
+						}
+
+						Action = FFileChangeData::FCA_Removed;
+					}
+					else if ((Event->mask & IN_DELETE) || (Event->mask & IN_MOVED_FROM))
+					{
+						// if a directory was deleted/moved, invalidate watch descriptors associated with it (unwatch it)
+						if (Event->mask & IN_ISDIR)
+						{
+							UnwatchDirectoryTree(AffectedFile);
+							// to be in sync with Windows implementation, ignore events about deleting directories unless told so
+							bIgnoreEvent = !bIncludeDirectoryChanges;
+						}
+
+						Action = FFileChangeData::FCA_Removed;
+					}
+
+					if (!bIgnoreEvent && Event->len)
+					{
+						new (FileChanges) FFileChangeData(AffectedFile, Action);
+					}
+				}
 			}
 
-			if ((Event->mask & IN_DELETE) || (Event->mask & IN_MOVED_FROM))
-			{
-				Action = FFileChangeData::FCA_Removed;
-			}
-
-			Ptr += sizeof(struct inotify_event) + Event->len;
-
-			const FString Filename = Directory / UTF8_TO_TCHAR(Event->name);
-			new (FileChanges) FFileChangeData(Filename, Action);
+			Ptr += EVENT_SIZE + Event->len;
 		}
 	}
 }
diff --git a/Engine/Source/Developer/DirectoryWatcher/Private/Linux/DirectoryWatchRequestLinux.h b/Engine/Source/Developer/DirectoryWatcher/Private/Linux/DirectoryWatchRequestLinux.h
--- a/Engine/Source/Developer/DirectoryWatcher/Private/Linux/DirectoryWatchRequestLinux.h
+++ b/Engine/Source/Developer/DirectoryWatcher/Private/Linux/DirectoryWatchRequestLinux.h
@@ -8,16 +8,38 @@
 #include <sys/inotify.h>
 
 #define EVENT_SIZE  ( sizeof (struct inotify_event) )
-#define BUF_LEN     ( 1024 * ( EVENT_SIZE + 16 ) )
+#define EVENT_BUF_LEN     ( 1024 * ( EVENT_SIZE + 16 ) )
 
 class FDirectoryWatchRequestLinux
 {
 public:
+
+	/**
+	 * Template class to support FString to TValueType maps with a case sensitive key
+	 */
+	template<typename TValueType>
+	struct FCaseSensitiveLookupKeyFuncs : BaseKeyFuncs<TValueType, FString>
+	{
+		static FORCEINLINE const FString& GetSetKey(const TPair<FString, TValueType>& Element)
+		{
+			return Element.Key;
+		}
+		static FORCEINLINE bool Matches(const FString& A, const FString& B)
+		{
+			return A.Equals(B, ESearchCase::CaseSensitive);
+		}
+		static FORCEINLINE uint32 GetKeyHash(const FString& Key)
+		{
+			return FCrc::StrCrc32<TCHAR>(*Key);
+		}
+	};
+
+
 	FDirectoryWatchRequestLinux();
 	virtual ~FDirectoryWatchRequestLinux();
 
 	/** Sets up the directory handle and request information */
-	bool Init(const FString& InDirectory);
+	bool Init(const FString& InDirectory, bool bInIncludeDirectoryChanges);
 
 	/** Adds a delegate to get fired when the directory changes */
 	FDelegateHandle AddDelegate( const IDirectoryWatcher::FDirectoryChanged& InDelegate );
@@ -37,16 +59,30 @@ public:
 
 private:
 
+	/** Adds watches for all files (and subdirectories) in a directory. */
+	void WatchDirectoryTree(const FString & RootAbsolutePath);
+
+	/** Removes all watches from a */
+	void UnwatchDirectoryTree(const FString & RootAbsolutePath);
+
 	FString Directory;
 
 	bool bRunning;
 	bool bEndWatchRequestInvoked;
 
+	/** Whether to report directory creation/deletion changes. */
+	bool bIncludeDirectoryChanges;
+
 	int FileDescriptor;
-	int * WatchDescriptor;
+
+	/** Mapping from watch descriptors to their path names */
+	TMap<int32, FString> WatchDescriptorsToPaths;
+
+	/** Mapping from paths to watch descriptors */
+	TMap<FString, int32, FDefaultSetAllocator, FCaseSensitiveLookupKeyFuncs<int32>> PathsToWatchDescriptors;
+
 	int NotifyFilter;
 
-	TArray<FString> AllFiles;
 	TArray<IDirectoryWatcher::FDirectoryChanged> Delegates;
 	TArray<FFileChangeData> FileChanges;
 
diff --git a/Engine/Source/Developer/DirectoryWatcher/Private/Linux/DirectoryWatcherLinux.cpp b/Engine/Source/Developer/DirectoryWatcher/Private/Linux/DirectoryWatcherLinux.cpp
--- a/Engine/Source/Developer/DirectoryWatcher/Private/Linux/DirectoryWatcherLinux.cpp
+++ b/Engine/Source/Developer/DirectoryWatcher/Private/Linux/DirectoryWatcherLinux.cpp
@@ -2,8 +2,6 @@
 
 #include "DirectoryWatcherPrivatePCH.h"
 
-#define WITH_DIRECTORY_WATCHER		1
-
 FDirectoryWatcherLinux::FDirectoryWatcherLinux()
 {
 	NumRequests = 0;
@@ -11,7 +9,6 @@ FDirectoryWatcherLinux::FDirectoryWatcherLinux()
 
 FDirectoryWatcherLinux::~FDirectoryWatcherLinux()
 {
-#if WITH_DIRECTORY_WATCHER
 	if (RequestMap.Num() != 0)
 	{
 		// Delete any remaining requests here. These requests are likely from modules which are still loaded at the time that this module unloads.
@@ -35,16 +32,13 @@ FDirectoryWatcherLinux::~FDirectoryWatcherLinux()
 			NumRequests--;
 		}
 	}
-#endif // WITH_DIRECTORY_WATCHER
-	
+
 	// Make sure every request that was created is destroyed
 	ensure(NumRequests == 0);
 }
 
 bool FDirectoryWatcherLinux::RegisterDirectoryChangedCallback( const FString& Directory, const FDirectoryChanged& InDelegate )
 {
-#if WITH_DIRECTORY_WATCHER
-	
 	FDirectoryWatchRequestLinux** RequestPtr = RequestMap.Find(Directory);
 	FDirectoryWatchRequestLinux* Request = NULL;
 
@@ -61,7 +55,7 @@ bool FDirectoryWatcherLinux::RegisterDirectoryChangedCallback( const FString& Di
 		NumRequests++;
 
 		// Begin reading directory changes
-		if (!Request->Init(Directory))
+		if (!Request->Init(Directory, false))
 		{
 			UE_LOG(LogDirectoryWatcher, Warning, TEXT("Failed to begin reading directory changes for %s."), *Directory);
 			delete Request;
@@ -75,17 +69,10 @@ bool FDirectoryWatcherLinux::RegisterDirectoryChangedCallback( const FString& Di
 	Request->AddDelegate(InDelegate);
 
 	return true;
-
-#else
-	
-	return false;
-	
-#endif // WITH_DIRECTORY_WATCHER
 }
 
 bool FDirectoryWatcherLinux::UnregisterDirectoryChangedCallback(const FString& Directory, const FDirectoryChanged& InDelegate)
 {
-#if WITH_DIRECTORY_WATCHER
 	FDirectoryWatchRequestLinux** RequestPtr = RequestMap.Find(Directory);
 
 	if (RequestPtr)
@@ -111,15 +98,12 @@ bool FDirectoryWatcherLinux::UnregisterDirectoryChangedCallback(const FString& D
 		}
 
 	}
-#endif // WITH_DIRECTORY_WATCHER
 
 	return false;
 }
 
 bool FDirectoryWatcherLinux::RegisterDirectoryChangedCallback_Handle( const FString& Directory, const FDirectoryChanged& InDelegate, FDelegateHandle& OutHandle, bool bIncludeDirectoryChanges )
 {
-#if WITH_DIRECTORY_WATCHER
-	
 	FDirectoryWatchRequestLinux** RequestPtr = RequestMap.Find(Directory);
 	FDirectoryWatchRequestLinux* Request = NULL;
 
@@ -136,7 +120,7 @@ bool FDirectoryWatcherLinux::RegisterDirectoryChangedCallback_Handle( const FStr
 		NumRequests++;
 
 		// Begin reading directory changes
-		if (!Request->Init(Directory))
+		if (!Request->Init(Directory, bIncludeDirectoryChanges))
 		{
 			UE_LOG(LogDirectoryWatcher, Warning, TEXT("Failed to begin reading directory changes for %s."), *Directory);
 			delete Request;
@@ -150,17 +134,10 @@ bool FDirectoryWatcherLinux::RegisterDirectoryChangedCallback_Handle( const FStr
 	OutHandle = Request->AddDelegate(InDelegate);
 
 	return true;
-
-#else
-	
-	return false;
-	
-#endif // WITH_DIRECTORY_WATCHER
 }
 
 bool FDirectoryWatcherLinux::UnregisterDirectoryChangedCallback_Handle(const FString& Directory, FDelegateHandle InHandle)
 {
-#if WITH_DIRECTORY_WATCHER
 	FDirectoryWatchRequestLinux** RequestPtr = RequestMap.Find(Directory);
 
 	if (RequestPtr)
@@ -186,15 +163,12 @@ bool FDirectoryWatcherLinux::UnregisterDirectoryChangedCallback_Handle(const FSt
 		}
 
 	}
-#endif // WITH_DIRECTORY_WATCHER
 
 	return false;
 }
 
 void FDirectoryWatcherLinux::Tick(float DeltaSeconds)
 {
-#if WITH_DIRECTORY_WATCHER
-
 	// Delete unregistered requests
 	for (int32 RequestIdx = RequestsPendingDelete.Num() - 1; RequestIdx >= 0; --RequestIdx)
 	{
@@ -209,6 +183,4 @@ void FDirectoryWatcherLinux::Tick(float DeltaSeconds)
 	{
 		RequestIt.Value()->ProcessPendingNotifications();
 	}
-
-#endif // WITH_DIRECTORY_WATCHER
 }
