[INTEGRATE] merge ue4<-ue4-launcherDev cl 2595224

diff --git a/Engine/Source/Developer/DirectoryWatcher/Private/Windows/DirectoryWatchRequestWindows.cpp b/Engine/Source/Developer/DirectoryWatcher/Private/Windows/DirectoryWatchRequestWindows.cpp
--- a/Engine/Source/Developer/DirectoryWatcher/Private/Windows/DirectoryWatchRequestWindows.cpp
+++ b/Engine/Source/Developer/DirectoryWatcher/Private/Windows/DirectoryWatchRequestWindows.cpp
@@ -146,11 +146,12 @@ void FDirectoryWatchRequestWindows::EndWatchRequest()
 #else
 			CancelIo(DirectoryHandle);
 #endif
-			// Wait for cancel operation to finish before closing the handle
-			WaitForSingleObjectEx(DirectoryHandle, 1000, true);
-			
-			::CloseHandle(DirectoryHandle);
+			// Clear the handle so we don't setup any more requests, and wait for the operation to finish
+			HANDLE TempDirectoryHandle = DirectoryHandle;
 			DirectoryHandle = INVALID_HANDLE_VALUE;
+			WaitForSingleObjectEx(TempDirectoryHandle, 1000, true);
+			
+			::CloseHandle(TempDirectoryHandle);
 		}
 		else
 		{
