DirectoryWatcher: allow ignoring changes in subtrees.

- Improves performance in SlateFileDialogs when browsing folders like Perforce workspace.
- A new test is included, old method has been deprecated.
- Mac version is not implemented (JIRA pending).

diff --git a/Engine/Plugins/Editor/PluginBrowser/Source/PluginBrowser/Private/SPluginBrowser.cpp b/Engine/Plugins/Editor/PluginBrowser/Source/PluginBrowser/Private/SPluginBrowser.cpp
--- a/Engine/Plugins/Editor/PluginBrowser/Source/PluginBrowser/Private/SPluginBrowser.cpp
+++ b/Engine/Plugins/Editor/PluginBrowser/Source/PluginBrowser/Private/SPluginBrowser.cpp
@@ -35,7 +35,7 @@ void SPluginBrowser::Construct( const FArguments& Args )
 	for(const FString& WatchDirectoryName: WatchDirectoryNames)
 	{
 		FDelegateHandle Handle;
-		if(DirectoryWatcherModule.Get()->RegisterDirectoryChangedCallback_Handle(WatchDirectoryName, IDirectoryWatcher::FDirectoryChanged::CreateRaw(this, &SPluginBrowser::OnPluginDirectoryChanged), Handle, true))
+		if(DirectoryWatcherModule.Get()->RegisterDirectoryChangedCallback_Handle(WatchDirectoryName, IDirectoryWatcher::FDirectoryChanged::CreateRaw(this, &SPluginBrowser::OnPluginDirectoryChanged), Handle, IDirectoryWatcher::WatchOptions::IncludeDirectoryChanges))
 		{
 			WatchDirectories.Add(WatchDirectoryName, Handle);
 		}
@@ -262,4 +262,4 @@ void SPluginBrowser::BreadcrumbTrail_OnCrumbClicked( const TSharedPtr<FPluginCat
 
 
 
-#undef LOCTEXT_NAMESPACE
\ No newline at end of file
+#undef LOCTEXT_NAMESPACE
diff --git a/Engine/Source/Developer/DirectoryWatcher/Private/DirectoryWatcherTests.cpp b/Engine/Source/Developer/DirectoryWatcher/Private/DirectoryWatcherTests.cpp
--- a/Engine/Source/Developer/DirectoryWatcher/Private/DirectoryWatcherTests.cpp
+++ b/Engine/Source/Developer/DirectoryWatcher/Private/DirectoryWatcherTests.cpp
@@ -12,7 +12,7 @@ struct FDirectoryWatcherTestPayload
 	FString WorkingDir;
 	TMap<FString, FFileChangeData::EFileChangeAction> ReportedChanges;
 
-	FDirectoryWatcherTestPayload(const FString& InWorkingDir, bool bIncludeDirectoryEvents = false)
+	FDirectoryWatcherTestPayload(const FString& InWorkingDir, uint32 Flags = 0)
 		: WorkingDir(InWorkingDir)
 	{
 		IFileManager::Get().MakeDirectory(*WorkingDir, true);
@@ -21,7 +21,7 @@ struct FDirectoryWatcherTestPayload
 		if (IDirectoryWatcher* DirectoryWatcher = Module.Get())
 		{
 			auto Callback = IDirectoryWatcher::FDirectoryChanged::CreateRaw(this, &FDirectoryWatcherTestPayload::OnDirectoryChanged);
-			DirectoryWatcher->RegisterDirectoryChangedCallback_Handle(WorkingDir, Callback, WatcherDelegate, bIncludeDirectoryEvents);
+			DirectoryWatcher->RegisterDirectoryChangedCallback_Handle(WorkingDir, Callback, WatcherDelegate, Flags);
 		}
 	}
 	~FDirectoryWatcherTestPayload()
@@ -275,7 +275,7 @@ bool FDirectoryWatcherNewFolderTest::RunTest(const FString& Parameters)
 		IFileManager::Get().MakeDirectory(*(WorkingDir / RemovedDirectory), true);
 
 		// Start watching the directory
-		TSharedPtr<FDirectoryWatcherTestPayload> Test = MakeShareable(new FDirectoryWatcherTestPayload(WorkingDir, true));
+		TSharedPtr<FDirectoryWatcherTestPayload> Test = MakeShareable(new FDirectoryWatcherTestPayload(WorkingDir, IDirectoryWatcher::WatchOptions::IncludeDirectoryChanges));
 
 		// Give the stream time to start up before doing the test
 		ADD_LATENT_AUTOMATION_COMMAND(FDelayedCallbackLatentCommand([=]{
@@ -306,12 +306,47 @@ bool FDirectoryWatcherNewFolderTest::RunTest(const FString& Parameters)
 	return true;
 }
 
+IMPLEMENT_SIMPLE_AUTOMATION_TEST(FDirectoryWatcherIgnoreSubtreeTest, "System.Plugins.Directory Watcher.Ignore Subtree", EAutomationTestFlags::ATF_Editor)
+bool FDirectoryWatcherIgnoreSubtreeTest::RunTest(const FString& Parameters)
+{
+	const FString WorkingDir = GetWorkingDir();
+
+	static const TCHAR* ChildDirectory = TEXT("child");
+	static const TCHAR* GrandchildDirectory = TEXT("grandchild");
+
+	// Give the stream time to start up before doing the test
+	ADD_LATENT_AUTOMATION_COMMAND(FDelayedCallbackLatentCommand([=]{
+
+		// Start watching the directory
+		TSharedPtr<FDirectoryWatcherTestPayload> Test = MakeShareable(new FDirectoryWatcherTestPayload(WorkingDir, IDirectoryWatcher::WatchOptions::IgnoreChangesInSubtree | IDirectoryWatcher::WatchOptions::IncludeDirectoryChanges));
+
+		// Give the stream time to start up before doing the test
+		ADD_LATENT_AUTOMATION_COMMAND(FDelayedCallbackLatentCommand([=]{
 
+			IFileManager::Get().MakeDirectory(*(WorkingDir / ChildDirectory), true);
+			IFileManager::Get().MakeDirectory(*(WorkingDir / ChildDirectory / GrandchildDirectory), true);
 
+			ADD_LATENT_AUTOMATION_COMMAND(FDelayedCallbackLatentCommand([=]{
 
+				FFileChangeData::EFileChangeAction* Action = Test->ReportedChanges.Find(ChildDirectory);
+				if (!Action || *Action != FFileChangeData::FCA_Added)
+				{
+					UE_LOG(LogDirectoryWatcherTests, Error, TEXT("Folder '%s' was not correctly reported as being added."), ChildDirectory);
+				}
 
+				Action = Test->ReportedChanges.Find(GrandchildDirectory);
+				if (Action)
+				{
+					UE_LOG(LogDirectoryWatcherTests, Error, TEXT("Changes to folder '%s' (creation of subfolder '%s') was reported in spite of us setting the mode 'ignore changes in subtree'."),
+						ChildDirectory, GrandchildDirectory);
+				}
 
+			}));
 
+		}));
 
+	}));
 
+	return true;
+}
 
diff --git a/Engine/Source/Developer/DirectoryWatcher/Private/Linux/DirectoryWatchRequestLinux.cpp b/Engine/Source/Developer/DirectoryWatcher/Private/Linux/DirectoryWatchRequestLinux.cpp
--- a/Engine/Source/Developer/DirectoryWatcher/Private/Linux/DirectoryWatchRequestLinux.cpp
+++ b/Engine/Source/Developer/DirectoryWatcher/Private/Linux/DirectoryWatchRequestLinux.cpp
@@ -5,6 +5,8 @@
 FDirectoryWatchRequestLinux::FDirectoryWatchRequestLinux()
 :	bRunning(false)
 ,	bEndWatchRequestInvoked(false)
+,	bIncludeDirectoryChanges(false)
+,	bWatchSubtree(false)
 {
 	NotifyFilter = IN_CREATE | IN_MOVE | IN_MODIFY | IN_DELETE;
 }
@@ -22,7 +24,7 @@ void FDirectoryWatchRequestLinux::Shutdown()
 	PathsToWatchDescriptors.Empty();
 }
 
-bool FDirectoryWatchRequestLinux::Init(const FString& InDirectory, bool bInIncludeDirectoryChanges)
+bool FDirectoryWatchRequestLinux::Init(const FString& InDirectory, uint32 Flags)
 {
 	if (InDirectory.Len() == 0)
 	{
@@ -31,7 +33,8 @@ bool FDirectoryWatchRequestLinux::Init(const FString& InDirectory, bool bInInclu
 	}
 
 	Directory = InDirectory;
-	bIncludeDirectoryChanges = bInIncludeDirectoryChanges;
+	bIncludeDirectoryChanges = (Flags & IDirectoryWatcher::WatchOptions::IncludeDirectoryChanges) != 0;
+	bWatchSubtree = (Flags & IDirectoryWatcher::WatchOptions::IgnoreChangesInSubtree) == 0;
 
 	if (bRunning)
 	{
@@ -117,7 +120,10 @@ void FDirectoryWatchRequestLinux::WatchDirectoryTree(const FString & RootAbsolut
 {
 	UE_LOG(LogDirectoryWatcher, VeryVerbose, TEXT("Watching tree '%s'"), *RootAbsolutePath);
 	TArray<FString> AllFiles;
-	IFileManager::Get().FindFilesRecursive(AllFiles, *RootAbsolutePath, TEXT("*"), false, true);
+	if (bWatchSubtree)
+	{
+		IFileManager::Get().FindFilesRecursive(AllFiles, *RootAbsolutePath, TEXT("*"), false, true);
+	}
 	// add the path as well
 	AllFiles.Add(RootAbsolutePath);
 
diff --git a/Engine/Source/Developer/DirectoryWatcher/Private/Linux/DirectoryWatchRequestLinux.h b/Engine/Source/Developer/DirectoryWatcher/Private/Linux/DirectoryWatchRequestLinux.h
--- a/Engine/Source/Developer/DirectoryWatcher/Private/Linux/DirectoryWatchRequestLinux.h
+++ b/Engine/Source/Developer/DirectoryWatcher/Private/Linux/DirectoryWatchRequestLinux.h
@@ -39,7 +39,7 @@ public:
 	virtual ~FDirectoryWatchRequestLinux();
 
 	/** Sets up the directory handle and request information */
-	bool Init(const FString& InDirectory, bool bInIncludeDirectoryChanges);
+	bool Init(const FString& InDirectory, uint32 Flags);
 
 	/** Adds a delegate to get fired when the directory changes */
 	FDelegateHandle AddDelegate( const IDirectoryWatcher::FDirectoryChanged& InDelegate );
@@ -73,6 +73,9 @@ private:
 	/** Whether to report directory creation/deletion changes. */
 	bool bIncludeDirectoryChanges;
 
+	/** Whether or not watch subtree. */
+	bool bWatchSubtree;
+
 	int FileDescriptor;
 
 	/** Mapping from watch descriptors to their path names */
diff --git a/Engine/Source/Developer/DirectoryWatcher/Private/Linux/DirectoryWatcherLinux.cpp b/Engine/Source/Developer/DirectoryWatcher/Private/Linux/DirectoryWatcherLinux.cpp
--- a/Engine/Source/Developer/DirectoryWatcher/Private/Linux/DirectoryWatcherLinux.cpp
+++ b/Engine/Source/Developer/DirectoryWatcher/Private/Linux/DirectoryWatcherLinux.cpp
@@ -102,7 +102,12 @@ bool FDirectoryWatcherLinux::UnregisterDirectoryChangedCallback(const FString& D
 	return false;
 }
 
-bool FDirectoryWatcherLinux::RegisterDirectoryChangedCallback_Handle( const FString& Directory, const FDirectoryChanged& InDelegate, FDelegateHandle& OutHandle, bool bIncludeDirectoryChanges )
+bool FDirectoryWatcherLinux::RegisterDirectoryChangedCallback_Handle (const FString& Directory, const FDirectoryChanged& InDelegate, FDelegateHandle& OutHandle, bool bIncludeDirectoryChanges)
+{
+	return RegisterDirectoryChangedCallback_Handle(Directory, InDelegate, OutHandle, bIncludeDirectoryChanges ? IDirectoryWatcher::WatchOptions::IncludeDirectoryChanges : 0);
+}
+
+bool FDirectoryWatcherLinux::RegisterDirectoryChangedCallback_Handle(const FString& Directory, const FDirectoryChanged& InDelegate, FDelegateHandle& OutHandle, uint32 Flags)
 {
 	FDirectoryWatchRequestLinux** RequestPtr = RequestMap.Find(Directory);
 	FDirectoryWatchRequestLinux* Request = NULL;
@@ -120,7 +125,7 @@ bool FDirectoryWatcherLinux::RegisterDirectoryChangedCallback_Handle( const FStr
 		NumRequests++;
 
 		// Begin reading directory changes
-		if (!Request->Init(Directory, bIncludeDirectoryChanges))
+		if (!Request->Init(Directory, Flags))
 		{
 			UE_LOG(LogDirectoryWatcher, Warning, TEXT("Failed to begin reading directory changes for %s."), *Directory);
 			delete Request;
@@ -136,6 +141,7 @@ bool FDirectoryWatcherLinux::RegisterDirectoryChangedCallback_Handle( const FStr
 	return true;
 }
 
+
 bool FDirectoryWatcherLinux::UnregisterDirectoryChangedCallback_Handle(const FString& Directory, FDelegateHandle InHandle)
 {
 	FDirectoryWatchRequestLinux** RequestPtr = RequestMap.Find(Directory);
diff --git a/Engine/Source/Developer/DirectoryWatcher/Private/Linux/DirectoryWatcherLinux.h b/Engine/Source/Developer/DirectoryWatcher/Private/Linux/DirectoryWatcherLinux.h
--- a/Engine/Source/Developer/DirectoryWatcher/Private/Linux/DirectoryWatcherLinux.h
+++ b/Engine/Source/Developer/DirectoryWatcher/Private/Linux/DirectoryWatcherLinux.h
@@ -10,7 +10,8 @@ public:
 
 	virtual bool RegisterDirectoryChangedCallback(const FString& Directory, const FDirectoryChanged& InDelegate) override;
 	virtual bool UnregisterDirectoryChangedCallback(const FString& Directory, const FDirectoryChanged& InDelegate) override;
-	virtual bool RegisterDirectoryChangedCallback_Handle(const FString& Directory, const FDirectoryChanged& InDelegate, FDelegateHandle& OutHandle, bool bIncludeDirectoryChanges) override;
+	virtual bool RegisterDirectoryChangedCallback_Handle (const FString& Directory, const FDirectoryChanged& InDelegate, FDelegateHandle& OutHandle, bool bIncludeDirectoryChanges) override;
+	virtual bool RegisterDirectoryChangedCallback_Handle (const FString& Directory, const FDirectoryChanged& InDelegate, FDelegateHandle& OutHandle, uint32 Flags) override;
 	virtual bool UnregisterDirectoryChangedCallback_Handle(const FString& Directory, FDelegateHandle InHandle) override;
 	virtual void Tick (float DeltaSeconds) override;
 
diff --git a/Engine/Source/Developer/DirectoryWatcher/Private/Mac/DirectoryWatchRequestMac.cpp b/Engine/Source/Developer/DirectoryWatcher/Private/Mac/DirectoryWatchRequestMac.cpp
--- a/Engine/Source/Developer/DirectoryWatcher/Private/Mac/DirectoryWatchRequestMac.cpp
+++ b/Engine/Source/Developer/DirectoryWatcher/Private/Mac/DirectoryWatchRequestMac.cpp
@@ -13,10 +13,11 @@ void DirectoryWatchMacCallback( ConstFSEventStreamRef StreamRef, void* WatchRequ
 
 // ============================================================================================================================
 
-FDirectoryWatchRequestMac::FDirectoryWatchRequestMac(bool bInIncludeDirectoryEvents)
+FDirectoryWatchRequestMac::FDirectoryWatchRequestMac(uint32 Flags)
 :	bRunning(false)
 ,	bEndWatchRequestInvoked(false)
-,	bIncludeDirectoryEvents(bInIncludeDirectoryEvents)
+,	bIncludeDirectoryEvents((Flags & IDirectoryWatcher::WatchOptions::IncludeDirectoryChanges) != 0)
+,	bIgnoreChangesInSubtree((Flags & IDirectoryWatcher::WatchOptions::IgnoreChangesInSubtree) != 0)
 {
 }
 
diff --git a/Engine/Source/Developer/DirectoryWatcher/Private/Mac/DirectoryWatchRequestMac.h b/Engine/Source/Developer/DirectoryWatcher/Private/Mac/DirectoryWatchRequestMac.h
--- a/Engine/Source/Developer/DirectoryWatcher/Private/Mac/DirectoryWatchRequestMac.h
+++ b/Engine/Source/Developer/DirectoryWatcher/Private/Mac/DirectoryWatchRequestMac.h
@@ -8,7 +8,7 @@
 class FDirectoryWatchRequestMac
 {
 public:
-	FDirectoryWatchRequestMac(bool bInIncludeDirectoryEvents);
+	FDirectoryWatchRequestMac(uint32 Flags);
 	virtual ~FDirectoryWatchRequestMac();
 
 	/** Sets up the directory handle and request information */
@@ -36,6 +36,7 @@ private:
 	bool				bRunning;
 	bool				bEndWatchRequestInvoked;
 	bool				bIncludeDirectoryEvents;
+	bool				bIgnoreChangesInSubtree;
 
 	TArray<IDirectoryWatcher::FDirectoryChanged> Delegates;
 	TArray<FFileChangeData> FileChanges;
diff --git a/Engine/Source/Developer/DirectoryWatcher/Private/Mac/DirectoryWatcherMac.cpp b/Engine/Source/Developer/DirectoryWatcher/Private/Mac/DirectoryWatcherMac.cpp
--- a/Engine/Source/Developer/DirectoryWatcher/Private/Mac/DirectoryWatcherMac.cpp
+++ b/Engine/Source/Developer/DirectoryWatcher/Private/Mac/DirectoryWatcherMac.cpp
@@ -102,7 +102,12 @@ bool FDirectoryWatcherMac::UnregisterDirectoryChangedCallback( const FString& Di
 	return false;
 }
 
-bool FDirectoryWatcherMac::RegisterDirectoryChangedCallback_Handle( const FString& Directory, const FDirectoryChanged& InDelegate, FDelegateHandle& OutHandle, bool bIncludeDirectoryChanges )
+bool FDirectoryWatcherMac::RegisterDirectoryChangedCallback_Handle (const FString& Directory, const FDirectoryChanged& InDelegate, FDelegateHandle& OutHandle, bool bIncludeDirectoryChanges)
+{
+	return RegisterDirectoryChangedCallback_Handle(Directory, InDelegate, OutHandle, bIncludeDirectoryChanges ? IDirectoryWatcher::WatchOptions::IncludeDirectoryChanges : 0);
+}
+
+bool FDirectoryWatcherMac::RegisterDirectoryChangedCallback_Handle( const FString& Directory, const FDirectoryChanged& InDelegate, FDelegateHandle& OutHandle, uint32 Flags )
 {
 	FDirectoryWatchRequestMac** RequestPtr = RequestMap.Find(Directory);
 	FDirectoryWatchRequestMac* Request = NULL;
@@ -116,7 +121,7 @@ bool FDirectoryWatcherMac::RegisterDirectoryChangedCallback_Handle( const FStrin
 	}
 	else
 	{
-		Request = new FDirectoryWatchRequestMac(bIncludeDirectoryChanges);
+		Request = new FDirectoryWatchRequestMac(Flags);
 		NumRequests++;
 
 		// Begin reading directory changes
diff --git a/Engine/Source/Developer/DirectoryWatcher/Private/Mac/DirectoryWatcherMac.h b/Engine/Source/Developer/DirectoryWatcher/Private/Mac/DirectoryWatcherMac.h
--- a/Engine/Source/Developer/DirectoryWatcher/Private/Mac/DirectoryWatcherMac.h
+++ b/Engine/Source/Developer/DirectoryWatcher/Private/Mac/DirectoryWatcherMac.h
@@ -12,6 +12,7 @@ public:
 	virtual bool RegisterDirectoryChangedCallback (const FString& Directory, const FDirectoryChanged& InDelegate) override;
 	virtual bool UnregisterDirectoryChangedCallback (const FString& Directory, const FDirectoryChanged& InDelegate) override;
 	virtual bool RegisterDirectoryChangedCallback_Handle (const FString& Directory, const FDirectoryChanged& InDelegate, FDelegateHandle& OutHandle, bool bIncludeDirectoryChanges) override;
+	virtual bool RegisterDirectoryChangedCallback_Handle (const FString& Directory, const FDirectoryChanged& InDelegate, FDelegateHandle& OutHandle, uint32 Flags) override;
 	virtual bool UnregisterDirectoryChangedCallback_Handle (const FString& Directory, FDelegateHandle InHandle) override;
 	virtual void Tick (float DeltaSeconds) override;
 
diff --git a/Engine/Source/Developer/DirectoryWatcher/Private/Windows/DirectoryWatchRequestWindows.cpp b/Engine/Source/Developer/DirectoryWatcher/Private/Windows/DirectoryWatchRequestWindows.cpp
--- a/Engine/Source/Developer/DirectoryWatcher/Private/Windows/DirectoryWatchRequestWindows.cpp
+++ b/Engine/Source/Developer/DirectoryWatcher/Private/Windows/DirectoryWatchRequestWindows.cpp
@@ -2,13 +2,14 @@
 
 #include "DirectoryWatcherPrivatePCH.h"
 
-FDirectoryWatchRequestWindows::FDirectoryWatchRequestWindows(bool bIncludeDirectoryEvents)
+FDirectoryWatchRequestWindows::FDirectoryWatchRequestWindows(uint32 Flags)
 {
 	bPendingDelete = false;
 	bEndWatchRequestInvoked = false;
 
 	MaxChanges = 16384;
-	bWatchSubtree = true;
+	bWatchSubtree = (Flags & IDirectoryWatcher::WatchOptions::IgnoreChangesInSubtree) == 0;
+	bool bIncludeDirectoryEvents = (Flags & IDirectoryWatcher::WatchOptions::IncludeDirectoryChanges) != 0;
 
 	NotifyFilter = FILE_NOTIFY_CHANGE_FILE_NAME | (bIncludeDirectoryEvents? FILE_NOTIFY_CHANGE_DIR_NAME : 0) | FILE_NOTIFY_CHANGE_LAST_WRITE | FILE_NOTIFY_CHANGE_CREATION;
 
diff --git a/Engine/Source/Developer/DirectoryWatcher/Private/Windows/DirectoryWatchRequestWindows.h b/Engine/Source/Developer/DirectoryWatcher/Private/Windows/DirectoryWatchRequestWindows.h
--- a/Engine/Source/Developer/DirectoryWatcher/Private/Windows/DirectoryWatchRequestWindows.h
+++ b/Engine/Source/Developer/DirectoryWatcher/Private/Windows/DirectoryWatchRequestWindows.h
@@ -5,7 +5,7 @@
 class FDirectoryWatchRequestWindows
 {
 public:
-	FDirectoryWatchRequestWindows(bool bIncludeDirectoryEvents);
+	FDirectoryWatchRequestWindows(uint32 Flags);
 	virtual ~FDirectoryWatchRequestWindows();
 
 	/** Sets up the directory handle and request information */
@@ -52,4 +52,4 @@ private:
 
 	TArray<IDirectoryWatcher::FDirectoryChanged> Delegates;
 	TArray<FFileChangeData> FileChanges;
-};
\ No newline at end of file
+};
diff --git a/Engine/Source/Developer/DirectoryWatcher/Private/Windows/DirectoryWatcherWindows.cpp b/Engine/Source/Developer/DirectoryWatcher/Private/Windows/DirectoryWatcherWindows.cpp
--- a/Engine/Source/Developer/DirectoryWatcher/Private/Windows/DirectoryWatcherWindows.cpp
+++ b/Engine/Source/Developer/DirectoryWatcher/Private/Windows/DirectoryWatcherWindows.cpp
@@ -105,7 +105,12 @@ bool FDirectoryWatcherWindows::UnregisterDirectoryChangedCallback( const FString
 	return false;
 }
 
-bool FDirectoryWatcherWindows::RegisterDirectoryChangedCallback_Handle( const FString& Directory, const FDirectoryChanged& InDelegate, FDelegateHandle& Handle, bool bIncludeDirectoryChanges )
+bool FDirectoryWatcherWindows::RegisterDirectoryChangedCallback_Handle (const FString& Directory, const FDirectoryChanged& InDelegate, FDelegateHandle& OutHandle, bool bIncludeDirectoryChanges)
+{
+	return RegisterDirectoryChangedCallback_Handle(Directory, InDelegate, OutHandle, bIncludeDirectoryChanges ? IDirectoryWatcher::WatchOptions::IncludeDirectoryChanges : 0);
+}
+
+bool FDirectoryWatcherWindows::RegisterDirectoryChangedCallback_Handle( const FString& Directory, const FDirectoryChanged& InDelegate, FDelegateHandle& Handle, uint32 Flags )
 {
 	FDirectoryWatchRequestWindows** RequestPtr = RequestMap.Find(Directory);
 	FDirectoryWatchRequestWindows* Request = NULL;
@@ -119,7 +124,7 @@ bool FDirectoryWatcherWindows::RegisterDirectoryChangedCallback_Handle( const FS
 	}
 	else
 	{
-		Request = new FDirectoryWatchRequestWindows(bIncludeDirectoryChanges);
+		Request = new FDirectoryWatchRequestWindows(Flags);
 		NumRequests++;
 
 		// Begin reading directory changes
diff --git a/Engine/Source/Developer/DirectoryWatcher/Private/Windows/DirectoryWatcherWindows.h b/Engine/Source/Developer/DirectoryWatcher/Private/Windows/DirectoryWatcherWindows.h
--- a/Engine/Source/Developer/DirectoryWatcher/Private/Windows/DirectoryWatcherWindows.h
+++ b/Engine/Source/Developer/DirectoryWatcher/Private/Windows/DirectoryWatcherWindows.h
@@ -17,7 +17,8 @@ public:
 
 	virtual bool RegisterDirectoryChangedCallback (const FString& Directory, const FDirectoryChanged& InDelegate) override;
 	virtual bool UnregisterDirectoryChangedCallback (const FString& Directory, const FDirectoryChanged& InDelegate) override;
-	virtual bool RegisterDirectoryChangedCallback_Handle (const FString& Directory, const FDirectoryChanged& InDelegate, FDelegateHandle& OutHandle, bool bIncludeDirectoryEvents) override;
+	virtual bool RegisterDirectoryChangedCallback_Handle (const FString& Directory, const FDirectoryChanged& InDelegate, FDelegateHandle& OutHandle, bool bIncludeDirectoryChanges) override;
+	virtual bool RegisterDirectoryChangedCallback_Handle (const FString& Directory, const FDirectoryChanged& InDelegate, FDelegateHandle& OutHandle, uint32 Flags) override;
 	virtual bool UnregisterDirectoryChangedCallback_Handle (const FString& Directory, FDelegateHandle InHandle) override;
 	virtual void Tick (float DeltaSeconds) override;
 
diff --git a/Engine/Source/Developer/DirectoryWatcher/Public/IDirectoryWatcher.h b/Engine/Source/Developer/DirectoryWatcher/Public/IDirectoryWatcher.h
--- a/Engine/Source/Developer/DirectoryWatcher/Public/IDirectoryWatcher.h
+++ b/Engine/Source/Developer/DirectoryWatcher/Public/IDirectoryWatcher.h
@@ -28,6 +28,17 @@ struct FFileChangeData
 class IDirectoryWatcher
 {
 public:
+
+	/** Options for a single watch (can be combined) */
+	enum WatchOptions : uint32
+	{
+		/** Whether to include notifications for changes to actual directories (such as directories being created or removed). */
+		IncludeDirectoryChanges			=		(1<<0),
+
+		/** Whether changes in subdirectories need to be reported. */
+		IgnoreChangesInSubtree			=		(1<<1),
+	};
+
 	/** A delegate to report directory changes */
 	DECLARE_DELEGATE_OneParam(FDirectoryChanged, const TArray<struct FFileChangeData>& /*FileChanges*/);
 
@@ -60,7 +71,18 @@ public:
 	 * @param	The handle to the registered delegate, if the registration was successful.
 	 * @param   Whether to include notifications for changes to actual directories (such as directories being created or removed)
 	 */
-	virtual bool RegisterDirectoryChangedCallback_Handle (const FString& Directory, const FDirectoryChanged& InDelegate, FDelegateHandle& OutHandle, bool bIncludeDirectoryChanges = false) = 0;
+	DEPRECATED(4.9, "Please use RegisterDirectoryChangedCallback_Handle that takes Flags argument") \
+	virtual bool RegisterDirectoryChangedCallback_Handle (const FString& Directory, const FDirectoryChanged& InDelegate, FDelegateHandle& OutHandle, bool bIncludeDirectoryChanges) = 0;
+
+	/**
+	 * Register a callback to fire when directories are changed
+	 *
+	 * @param	Directory to watch
+	 * @param	Delegate to add to our callback list
+	 * @param	The handle to the registered delegate, if the registration was successful.
+	 * @param	Set of options to use when registering the delegates.
+	 */
+	virtual bool RegisterDirectoryChangedCallback_Handle (const FString& Directory, const FDirectoryChanged& InDelegate, FDelegateHandle& OutHandle, uint32 Flags = 0) = 0;
 
 	/**
 	 * Unregisters a callback to fire when directories are changed
