DirectoryWatcher: deleted the deprecated method.

- If someone was passing something like TRUE/FALSE instead of booleans to the deprecated method this would result in silent brekage.

diff --git a/Engine/Source/Developer/DirectoryWatcher/Private/Linux/DirectoryWatcherLinux.cpp b/Engine/Source/Developer/DirectoryWatcher/Private/Linux/DirectoryWatcherLinux.cpp
--- a/Engine/Source/Developer/DirectoryWatcher/Private/Linux/DirectoryWatcherLinux.cpp
+++ b/Engine/Source/Developer/DirectoryWatcher/Private/Linux/DirectoryWatcherLinux.cpp
@@ -102,11 +102,6 @@ bool FDirectoryWatcherLinux::UnregisterDirectoryChangedCallback(const FString& D
 	return false;
 }
 
-bool FDirectoryWatcherLinux::RegisterDirectoryChangedCallback_Handle (const FString& Directory, const FDirectoryChanged& InDelegate, FDelegateHandle& OutHandle, bool bIncludeDirectoryChanges)
-{
-	return RegisterDirectoryChangedCallback_Handle(Directory, InDelegate, OutHandle, bIncludeDirectoryChanges ? IDirectoryWatcher::WatchOptions::IncludeDirectoryChanges : 0);
-}
-
 bool FDirectoryWatcherLinux::RegisterDirectoryChangedCallback_Handle(const FString& Directory, const FDirectoryChanged& InDelegate, FDelegateHandle& OutHandle, uint32 Flags)
 {
 	FDirectoryWatchRequestLinux** RequestPtr = RequestMap.Find(Directory);
diff --git a/Engine/Source/Developer/DirectoryWatcher/Private/Linux/DirectoryWatcherLinux.h b/Engine/Source/Developer/DirectoryWatcher/Private/Linux/DirectoryWatcherLinux.h
--- a/Engine/Source/Developer/DirectoryWatcher/Private/Linux/DirectoryWatcherLinux.h
+++ b/Engine/Source/Developer/DirectoryWatcher/Private/Linux/DirectoryWatcherLinux.h
@@ -10,7 +10,6 @@ public:
 
 	virtual bool RegisterDirectoryChangedCallback(const FString& Directory, const FDirectoryChanged& InDelegate) override;
 	virtual bool UnregisterDirectoryChangedCallback(const FString& Directory, const FDirectoryChanged& InDelegate) override;
-	virtual bool RegisterDirectoryChangedCallback_Handle (const FString& Directory, const FDirectoryChanged& InDelegate, FDelegateHandle& OutHandle, bool bIncludeDirectoryChanges) override;
 	virtual bool RegisterDirectoryChangedCallback_Handle (const FString& Directory, const FDirectoryChanged& InDelegate, FDelegateHandle& OutHandle, uint32 Flags) override;
 	virtual bool UnregisterDirectoryChangedCallback_Handle(const FString& Directory, FDelegateHandle InHandle) override;
 	virtual void Tick (float DeltaSeconds) override;
diff --git a/Engine/Source/Developer/DirectoryWatcher/Private/Mac/DirectoryWatcherMac.cpp b/Engine/Source/Developer/DirectoryWatcher/Private/Mac/DirectoryWatcherMac.cpp
--- a/Engine/Source/Developer/DirectoryWatcher/Private/Mac/DirectoryWatcherMac.cpp
+++ b/Engine/Source/Developer/DirectoryWatcher/Private/Mac/DirectoryWatcherMac.cpp
@@ -102,11 +102,6 @@ bool FDirectoryWatcherMac::UnregisterDirectoryChangedCallback( const FString& Di
 	return false;
 }
 
-bool FDirectoryWatcherMac::RegisterDirectoryChangedCallback_Handle (const FString& Directory, const FDirectoryChanged& InDelegate, FDelegateHandle& OutHandle, bool bIncludeDirectoryChanges)
-{
-	return RegisterDirectoryChangedCallback_Handle(Directory, InDelegate, OutHandle, bIncludeDirectoryChanges ? IDirectoryWatcher::WatchOptions::IncludeDirectoryChanges : 0);
-}
-
 bool FDirectoryWatcherMac::RegisterDirectoryChangedCallback_Handle( const FString& Directory, const FDirectoryChanged& InDelegate, FDelegateHandle& OutHandle, uint32 Flags )
 {
 	FDirectoryWatchRequestMac** RequestPtr = RequestMap.Find(Directory);
diff --git a/Engine/Source/Developer/DirectoryWatcher/Private/Mac/DirectoryWatcherMac.h b/Engine/Source/Developer/DirectoryWatcher/Private/Mac/DirectoryWatcherMac.h
--- a/Engine/Source/Developer/DirectoryWatcher/Private/Mac/DirectoryWatcherMac.h
+++ b/Engine/Source/Developer/DirectoryWatcher/Private/Mac/DirectoryWatcherMac.h
@@ -11,7 +11,6 @@ public:
 
 	virtual bool RegisterDirectoryChangedCallback (const FString& Directory, const FDirectoryChanged& InDelegate) override;
 	virtual bool UnregisterDirectoryChangedCallback (const FString& Directory, const FDirectoryChanged& InDelegate) override;
-	virtual bool RegisterDirectoryChangedCallback_Handle (const FString& Directory, const FDirectoryChanged& InDelegate, FDelegateHandle& OutHandle, bool bIncludeDirectoryChanges) override;
 	virtual bool RegisterDirectoryChangedCallback_Handle (const FString& Directory, const FDirectoryChanged& InDelegate, FDelegateHandle& OutHandle, uint32 Flags) override;
 	virtual bool UnregisterDirectoryChangedCallback_Handle (const FString& Directory, FDelegateHandle InHandle) override;
 	virtual void Tick (float DeltaSeconds) override;
diff --git a/Engine/Source/Developer/DirectoryWatcher/Private/Windows/DirectoryWatcherWindows.cpp b/Engine/Source/Developer/DirectoryWatcher/Private/Windows/DirectoryWatcherWindows.cpp
--- a/Engine/Source/Developer/DirectoryWatcher/Private/Windows/DirectoryWatcherWindows.cpp
+++ b/Engine/Source/Developer/DirectoryWatcher/Private/Windows/DirectoryWatcherWindows.cpp
@@ -105,11 +105,6 @@ bool FDirectoryWatcherWindows::UnregisterDirectoryChangedCallback( const FString
 	return false;
 }
 
-bool FDirectoryWatcherWindows::RegisterDirectoryChangedCallback_Handle (const FString& Directory, const FDirectoryChanged& InDelegate, FDelegateHandle& OutHandle, bool bIncludeDirectoryChanges)
-{
-	return RegisterDirectoryChangedCallback_Handle(Directory, InDelegate, OutHandle, bIncludeDirectoryChanges ? IDirectoryWatcher::WatchOptions::IncludeDirectoryChanges : 0);
-}
-
 bool FDirectoryWatcherWindows::RegisterDirectoryChangedCallback_Handle( const FString& Directory, const FDirectoryChanged& InDelegate, FDelegateHandle& Handle, uint32 Flags )
 {
 	FDirectoryWatchRequestWindows** RequestPtr = RequestMap.Find(Directory);
diff --git a/Engine/Source/Developer/DirectoryWatcher/Private/Windows/DirectoryWatcherWindows.h b/Engine/Source/Developer/DirectoryWatcher/Private/Windows/DirectoryWatcherWindows.h
--- a/Engine/Source/Developer/DirectoryWatcher/Private/Windows/DirectoryWatcherWindows.h
+++ b/Engine/Source/Developer/DirectoryWatcher/Private/Windows/DirectoryWatcherWindows.h
@@ -17,7 +17,6 @@ public:
 
 	virtual bool RegisterDirectoryChangedCallback (const FString& Directory, const FDirectoryChanged& InDelegate) override;
 	virtual bool UnregisterDirectoryChangedCallback (const FString& Directory, const FDirectoryChanged& InDelegate) override;
-	virtual bool RegisterDirectoryChangedCallback_Handle (const FString& Directory, const FDirectoryChanged& InDelegate, FDelegateHandle& OutHandle, bool bIncludeDirectoryChanges) override;
 	virtual bool RegisterDirectoryChangedCallback_Handle (const FString& Directory, const FDirectoryChanged& InDelegate, FDelegateHandle& OutHandle, uint32 Flags) override;
 	virtual bool UnregisterDirectoryChangedCallback_Handle (const FString& Directory, FDelegateHandle InHandle) override;
 	virtual void Tick (float DeltaSeconds) override;
diff --git a/Engine/Source/Developer/DirectoryWatcher/Public/IDirectoryWatcher.h b/Engine/Source/Developer/DirectoryWatcher/Public/IDirectoryWatcher.h
--- a/Engine/Source/Developer/DirectoryWatcher/Public/IDirectoryWatcher.h
+++ b/Engine/Source/Developer/DirectoryWatcher/Public/IDirectoryWatcher.h
@@ -63,17 +63,6 @@ public:
 	DELEGATE_DEPRECATED("The UnregisterDirectoryChangedCallback function is deprecated - please use the RegisterDirectoryChangedCallback_Handle function and pass its out parameter to UnregisterDirectoryChangedCallback_Handle instead.") \
 	virtual bool UnregisterDirectoryChangedCallback (const FString& Directory, const FDirectoryChanged& InDelegate) = 0;
 
-	/**
-	 * Register a callback to fire when directories are changed
-	 *
-	 * @param	Directory to watch
-	 * @param	Delegate to add to our callback list
-	 * @param	The handle to the registered delegate, if the registration was successful.
-	 * @param   Whether to include notifications for changes to actual directories (such as directories being created or removed)
-	 */
-	DEPRECATED(4.9, "Please use RegisterDirectoryChangedCallback_Handle that takes Flags argument") \
-	virtual bool RegisterDirectoryChangedCallback_Handle (const FString& Directory, const FDirectoryChanged& InDelegate, FDelegateHandle& OutHandle, bool bIncludeDirectoryChanges) = 0;
-
 	/**
 	 * Register a callback to fire when directories are changed
 	 *
