Moved the FileCache out of UnrealEd so it can be used without having to link to the editor

It now exists within the DirectoryWatcher namespace of the DirectoryWatcher module.

diff --git a/Engine/Source/Developer/DirectoryWatcher/Private/FileCache.cpp b/Engine/Source/Developer/DirectoryWatcher/Private/FileCache.cpp
new file mode 100644
--- /dev/null
+++ b/Engine/Source/Developer/DirectoryWatcher/Private/FileCache.cpp
@@ -0,0 +1,1021 @@
+// Copyright 1998-2015 Epic Games, Inc. All Rights Reserved.
+
+#include "DirectoryWatcherPrivatePCH.h"
+#include "FileCache.h"
+#include "DirectoryWatcherModule.h"
+
+namespace DirectoryWatcher
+{
+
+template<typename T>
+void ReadWithCustomVersions(FArchive& Ar, T& Data)
+{
+	int64 CustomVersionsOffset = 0;
+	Ar << CustomVersionsOffset;
+
+	const int64 DataStart = Ar.Tell();
+
+	Ar.Seek(CustomVersionsOffset);
+
+	// Serialize the custom versions
+	FCustomVersionContainer Vers = Ar.GetCustomVersions();
+	Vers.Serialize(Ar);
+	Ar.SetCustomVersions(Vers);
+
+	Ar.Seek(DataStart);
+
+	Ar << Data;
+}
+
+template<typename T>
+void WriteWithCustomVersions(FArchive& Ar, T& Data)
+{
+	const int64 CustomVersionsHeader = Ar.Tell();
+	int64 CustomVersionsOffset = CustomVersionsHeader;
+	// We'll come back later and fill this in
+	Ar << CustomVersionsOffset;
+
+	// Write out the data
+	Ar << Data;
+
+	CustomVersionsOffset = Ar.Tell();
+
+	// Serialize the custom versions
+	FCustomVersionContainer Vers = Ar.GetCustomVersions();
+	Vers.Serialize(Ar);
+
+	// Write out where the custom versions are in our header
+	Ar.Seek(CustomVersionsHeader);
+	Ar << CustomVersionsOffset;
+}
+
+/** Convert a FFileChangeData::EFileChangeAction into an EFileAction */
+EFileAction ToFileAction(FFileChangeData::EFileChangeAction InAction)
+{
+	switch (InAction)
+	{
+	case FFileChangeData::FCA_Added: 	return EFileAction::Added;
+	case FFileChangeData::FCA_Modified:	return EFileAction::Modified;
+	case FFileChangeData::FCA_Removed:	return EFileAction::Removed;
+	default:							return EFileAction::Modified;
+	}
+}
+
+const FGuid FFileCacheCustomVersion::Key(0x8E7DDCB3, 0x80DA47BB, 0x9FD346A2, 0x93984DF6);
+FCustomVersionRegistration GRegisterFileCacheVersion(FFileCacheCustomVersion::Key, FFileCacheCustomVersion::Latest, TEXT("FileCacheVersion"));
+
+static const uint32 CacheFileMagicNumber = 0x03DCCB00;
+
+/** Single runnable thread used to parse file cache directories without blocking the main thread */
+struct FAsyncTaskThread : public FRunnable
+{
+	typedef TArray<TWeakPtr<IAsyncFileCacheTask, ESPMode::ThreadSafe>> FTaskArray;
+
+	FAsyncTaskThread() : Thread(nullptr) {}
+
+	/** Add a reader to this thread which will get ticked periodically until complete */
+	void AddTask(TSharedPtr<IAsyncFileCacheTask, ESPMode::ThreadSafe> InTask)
+	{
+		FScopeLock Lock(&TaskArrayMutex);
+		Tasks.Add(InTask);
+
+		if (!Thread)
+		{
+			static int32 Index = 0;
+			Thread = FRunnableThread::Create(this, *FString::Printf(TEXT("AsyncTaskThread_%d"), ++Index));
+		}
+	}
+
+	/** Run this thread */
+	virtual uint32 Run()
+	{
+		for(;;)
+		{
+			// Copy the array while we tick the readers
+			FTaskArray Dupl;
+			{
+				FScopeLock Lock(&TaskArrayMutex);
+				Dupl = Tasks;
+			}
+
+			// Tick each one for a second
+			for (auto& Task : Dupl)
+			{
+				auto PinnedTask = Task.Pin();
+				if (PinnedTask.IsValid())
+				{
+					PinnedTask->Tick(FTimeLimit(1));
+				}
+			}
+
+			// Cleanup dead/finished Tasks
+			FScopeLock Lock(&TaskArrayMutex);
+			for (int32 Index = 0; Index < Tasks.Num(); )
+			{
+				auto Task = Tasks[Index].Pin();
+				if (!Task.IsValid() || Task->IsComplete())
+				{
+					Tasks.RemoveAt(Index);
+				}
+				else
+				{
+					++Index;
+				}
+			}
+
+			// Shutdown the thread if we've nothing left to do
+			if (Tasks.Num() == 0)
+			{
+				Thread = nullptr;
+				break;
+			}
+		}
+
+		return 0;
+	}
+
+private:
+	/** We start our own thread if one doesn't already exist. */
+	FRunnableThread* Thread;
+
+	/** Array of things that need ticking, and a mutex to protect them */
+	FCriticalSection TaskArrayMutex;
+	FTaskArray Tasks;
+};
+
+FAsyncTaskThread AsyncTaskThread;
+
+/** Threading strategy for FAsyncFileHasher:
+ *	The task is constructed on the main thread with its Data.
+ *	The array 'Data' *never* changes size. The task thread moves along setting file hashes, while the main thread
+ * 	trails behind accessing the completed entries. We should thus never have 2 threads accessing the same memory,
+ *	except for the atomic 'CurrentIndex'
+ */
+
+FAsyncFileHasher::FAsyncFileHasher(TArray<FFilenameAndHash> InFilesThatNeedHashing)
+	: Data(MoveTemp(InFilesThatNeedHashing)), NumReturned(0)
+{
+	// Read in files in 1MB chunks
+	ScratchBuffer.SetNumUninitialized(1024 * 1024);
+}
+
+TArray<FFilenameAndHash> FAsyncFileHasher::GetCompletedData()
+{
+	// Don't need to lock here since the thread will never look at the array before CurrentIndex.
+	TArray<FFilenameAndHash> Local;
+	const int32 CompletedIndex = CurrentIndex.GetValue();
+
+	if (NumReturned < CompletedIndex)
+	{
+		Local.Append(Data.GetData() + NumReturned, CompletedIndex - NumReturned);
+		NumReturned = CompletedIndex;
+
+		if (CompletedIndex == Data.Num())
+		{
+			Data.Empty();
+			CurrentIndex.Set(0);
+		}
+	}
+
+	return Local;
+}
+
+bool FAsyncFileHasher::IsComplete() const
+{
+	return CurrentIndex.GetValue() == Data.Num();
+}
+
+IAsyncFileCacheTask::EProgressResult FAsyncFileHasher::Tick(const FTimeLimit& Limit)
+{
+	for (; CurrentIndex.GetValue() < Data.Num(); )
+	{
+		const auto Index = CurrentIndex.GetValue();
+		Data[Index].FileHash = FMD5Hash::HashFile(*Data[Index].AbsoluteFilename, &ScratchBuffer);
+
+		CurrentIndex.Increment();
+
+		if (Limit.Exceeded())
+		{
+			return EProgressResult::Pending;
+		}
+	}
+
+	return EProgressResult::Finished;
+}
+
+/** Threading strategy for FAsyncDirectoryReader:
+ *	The directory reader owns the cached and live state until it has completely finished. Once IsComplete() is true, the main thread can 
+ *	have access to both the cached and farmed data.
+ */
+
+
+FAsyncDirectoryReader::FAsyncDirectoryReader(const FString& InDirectory, EPathType InPathType)
+	: RootPath(InDirectory), PathType(InPathType)
+{
+	PendingDirectories.Add(InDirectory);
+	LiveState.Emplace();
+}
+
+TOptional<FDirectoryState> FAsyncDirectoryReader::GetLiveState()
+{
+	TOptional<FDirectoryState> OldState;
+
+	if (ensureMsg(IsComplete(), TEXT("Invalid property access from thread before task completion")))
+	{
+		Swap(OldState, LiveState);
+	}
+
+	return OldState;
+}
+
+
+TOptional<FDirectoryState> FAsyncDirectoryReader::GetCachedState()
+{
+	TOptional<FDirectoryState> OldState;
+
+	if (ensureMsg(IsComplete(), TEXT("Invalid property access from thread before task completion")))
+	{
+		Swap(OldState, CachedState);
+	}
+
+	return OldState;
+}
+
+bool FAsyncDirectoryReader::IsComplete() const
+{
+	return bIsComplete;
+}
+
+FAsyncDirectoryReader::EProgressResult FAsyncDirectoryReader::Tick(const FTimeLimit& TimeLimit)
+{
+	if (IsComplete())
+	{
+		return EProgressResult::Finished;
+	}
+
+	auto& FileManager = IFileManager::Get();
+	const int32 RootPathLen = RootPath.Len();
+
+	// Discover files
+	for (int32 Index = 0; Index < PendingDirectories.Num(); ++Index)
+	{
+		ScanDirectory(PendingDirectories[Index]);
+		if (TimeLimit.Exceeded())
+		{
+			// We've spent too long, bail
+			PendingDirectories.RemoveAt(0, Index + 1, false);
+			return EProgressResult::Pending;
+		}
+	}
+	PendingDirectories.Empty();
+
+	// Process files
+	for (int32 Index = 0; Index < PendingFiles.Num(); ++Index)
+	{
+		const auto& File = PendingFiles[Index];
+
+		// Store the file relative or absolute
+		FString Filename = (PathType == EPathType::Relative ? *File + RootPathLen : *File);
+
+		const auto Timestamp = FileManager.GetTimeStamp(*File);
+	
+		FMD5Hash MD5;
+		if (CachedState.IsSet())
+		{
+			const FFileData* CachedData = CachedState->Files.Find(Filename);
+			if (CachedData && CachedData->Timestamp == Timestamp && CachedData->FileHash.IsValid())
+			{
+				// Use the cached MD5 to avoid opening the file
+				MD5 = CachedData->FileHash;
+			}
+		}
+
+		if (!MD5.IsValid())
+		{
+			FilesThatNeedHashing.Emplace(File);
+		}
+
+		LiveState->Files.Emplace(MoveTemp(Filename), FFileData(Timestamp, MD5));
+
+		if (TimeLimit.Exceeded())
+		{
+			// We've spent too long, bail
+			PendingFiles.RemoveAt(0, Index + 1, false);
+			return EProgressResult::Pending;
+		}
+	}
+	PendingFiles.Empty();
+
+	bIsComplete = true;
+
+	UE_LOG(LogFileCache, Log, TEXT("Scanning file cache for directory '%s' took %.2fs"), *RootPath, GetAge());
+	return EProgressResult::Finished;
+}
+
+void FAsyncDirectoryReader::ScanDirectory(const FString& InDirectory)
+{
+	struct FVisitor : public IPlatformFile::FDirectoryVisitor
+	{
+		TArray<FString>* PendingFiles;
+		TArray<FString>* PendingDirectories;
+		FMatchRules* Rules;
+		int32 RootPathLength;
+
+		virtual bool Visit(const TCHAR* FilenameOrDirectory, bool bIsDirectory)
+		{
+			FString FileStr(FilenameOrDirectory);
+			if (bIsDirectory)
+			{
+				PendingDirectories->Add(MoveTemp(FileStr));
+			}
+			else if (Rules->IsFileApplicable(FilenameOrDirectory + RootPathLength))
+			{
+				PendingFiles->Add(MoveTemp(FileStr));
+			}
+			return true;
+		}
+	};
+
+	FVisitor Visitor;
+	Visitor.PendingFiles = &PendingFiles;
+	Visitor.PendingDirectories = &PendingDirectories;
+	Visitor.Rules = &LiveState->Rules;
+	Visitor.RootPathLength = RootPath.Len();
+
+	IFileManager::Get().IterateDirectory(*InDirectory, Visitor);
+}
+
+FFileCache::FFileCache(const FFileCacheConfig& InConfig)
+	: Config(InConfig)
+	, bSavedCacheDirty(false)
+	, LastFileHashGetTime(0)
+{
+	// Ensure the directory has a trailing /
+	Config.Directory /= TEXT("");
+
+	// bDetectMoves implies bRequireFileHashes
+	Config.bRequireFileHashes = Config.bRequireFileHashes || Config.bDetectMoves;
+
+	DirectoryReader = MakeShareable(new FAsyncDirectoryReader(Config.Directory, Config.PathType));
+	DirectoryReader->SetMatchRules(Config.Rules);
+
+	// Attempt to load an existing cache file
+	auto ExistingCache = ReadCache();
+	if (ExistingCache.IsSet())
+	{
+		DirectoryReader->UseCachedState(MoveTemp(ExistingCache.GetValue()));
+	}
+
+	AsyncTaskThread.AddTask(DirectoryReader);
+
+	FDirectoryWatcherModule& Module = FModuleManager::LoadModuleChecked<FDirectoryWatcherModule>(TEXT("DirectoryWatcher"));
+	if (IDirectoryWatcher* DirectoryWatcher = Module.Get())
+	{
+		auto Callback = IDirectoryWatcher::FDirectoryChanged::CreateRaw(this, &FFileCache::OnDirectoryChanged);
+		DirectoryWatcher->RegisterDirectoryChangedCallback_Handle(Config.Directory, Callback, WatcherDelegate);
+	}
+}
+
+FFileCache::~FFileCache()
+{
+	UnbindWatcher();
+	WriteCache();
+}
+
+void FFileCache::Destroy()
+{
+	// Delete the cache file, and clear out everything
+	bSavedCacheDirty = false;
+	if (!Config.CacheFile.IsEmpty())
+	{
+		IFileManager::Get().Delete(*Config.CacheFile, false, true, true);
+	}
+
+	DirectoryReader = nullptr;
+	AsyncFileHasher = nullptr;
+	DirtyFileHasher = nullptr;
+
+	DirtyFiles.Empty();
+	CachedDirectoryState = FDirectoryState();
+
+	UnbindWatcher();
+}
+
+bool FFileCache::HasStartedUp() const
+{
+	return !DirectoryReader.IsValid() || DirectoryReader->IsComplete();
+}
+
+bool FFileCache::MoveDetectionInitialized() const
+{
+	if (!HasStartedUp())
+	{
+		return false;
+	}
+	else if (!Config.bDetectMoves)
+	{
+		return true;
+	}
+	else
+	{
+		// We don't check AsyncFileHasher->IsComplete() here because that doesn't necessarily mean we've harvested the results off the thread
+		return !AsyncFileHasher.IsValid();
+	}
+}
+
+const FFileData* FFileCache::FindFileData(FImmutableString InFilename) const
+{
+	if (!ensure(HasStartedUp()))
+	{
+		// It's invalid to call this while the cached state is still being updated on a thread.
+		return nullptr;
+	}
+
+	return CachedDirectoryState.Files.Find(InFilename);
+}
+
+void FFileCache::UnbindWatcher()
+{
+	if (!WatcherDelegate.IsValid())
+	{
+		return;
+	}
+
+	if (FDirectoryWatcherModule* Module = FModuleManager::GetModulePtr<FDirectoryWatcherModule>(TEXT("DirectoryWatcher")))
+	{
+		if (IDirectoryWatcher* DirectoryWatcher = Module->Get())
+		{
+			DirectoryWatcher->UnregisterDirectoryChangedCallback_Handle(Config.Directory, WatcherDelegate);
+		}
+	}
+
+	WatcherDelegate.Reset();
+}
+
+TOptional<FDirectoryState> FFileCache::ReadCache() const
+{
+	TOptional<FDirectoryState> Optional;
+	if (!Config.CacheFile.IsEmpty())
+	{		
+		FArchive* Ar = IFileManager::Get().CreateFileReader(*Config.CacheFile);
+		if (Ar)
+		{
+			// Serialize the magic number - the first iteration omitted version information, so we have a magic number to ignore this data
+			uint32 MagicNumber = 0;
+			*Ar << MagicNumber;
+
+			if (MagicNumber == CacheFileMagicNumber)
+			{
+				FDirectoryState Result;
+				ReadWithCustomVersions(*Ar, Result);
+
+				Optional.Emplace(MoveTemp(Result));
+			}
+
+			Ar->Close();
+			delete Ar;
+		}
+	}
+
+	return Optional;
+}
+
+void FFileCache::WriteCache()
+{
+	if (bSavedCacheDirty && !Config.CacheFile.IsEmpty())
+	{
+		const FString ParentFolder = FPaths::GetPath(Config.CacheFile);
+		if (!IFileManager::Get().DirectoryExists(*ParentFolder))
+		{
+			IFileManager::Get().MakeDirectory(*ParentFolder, true);	
+		}
+
+		// Write to a temp file to avoid corruption
+		FString TempFile = Config.CacheFile + TEXT(".tmp");
+		
+		FArchive* Ar = IFileManager::Get().CreateFileWriter(*TempFile);
+		if (ensureMsgf(Ar, TEXT("Unable to write file-cache for '%s' to '%s'."), *Config.Directory, *Config.CacheFile))
+		{
+			// Serialize the magic number
+			uint32 MagicNumber = CacheFileMagicNumber;
+			*Ar << MagicNumber;
+
+			WriteWithCustomVersions(*Ar, CachedDirectoryState);
+
+			Ar->Close();
+			delete Ar;
+
+			CachedDirectoryState.Files.Shrink();
+
+			bSavedCacheDirty = false;
+
+			const bool bMoved = IFileManager::Get().Move(*Config.CacheFile, *TempFile, true, true);
+			ensureMsgf(bMoved, TEXT("Unable to move file-cache for '%s' from '%s' to '%s'."), *Config.Directory, *TempFile, *Config.CacheFile);
+		}
+	}
+}
+
+FString FFileCache::GetAbsolutePath(const FString& InTransactionPath) const
+{
+	if (Config.PathType == EPathType::Relative)
+	{
+		return Config.Directory / InTransactionPath;
+	}
+	else
+	{
+		return InTransactionPath;
+	}
+}
+
+TOptional<FString> FFileCache::GetTransactionPath(const FString& InAbsolutePath) const
+{
+	FString Temp = FPaths::ConvertRelativePathToFull(InAbsolutePath);
+	FString RelativePath(*Temp + Config.Directory.Len());
+
+	// If it's a directory or is not applicable, ignore it
+	if (!Temp.StartsWith(Config.Directory) || IFileManager::Get().DirectoryExists(*Temp) || !Config.Rules.IsFileApplicable(*RelativePath))
+	{
+		return TOptional<FString>();
+	}
+
+	if (Config.PathType == EPathType::Relative)
+	{
+		return MoveTemp(RelativePath);
+	}
+	else
+	{
+		return MoveTemp(Temp);
+	}
+}
+
+void FFileCache::DiffDirtyFiles(TMap<FImmutableString, FFileData>& InDirtyFiles, TArray<FUpdateCacheTransaction>& OutTransactions, const FDirectoryState* InFileSystemState) const
+{
+	TArray<uint8> ScratchBuffer;
+	ScratchBuffer.SetNumUninitialized(1024*1024);
+
+	TMap<FImmutableString, FFileData> AddedFiles, ModifiedFiles;
+	TSet<FImmutableString> RemovedFiles, InvalidDirtyFiles;
+
+	auto& FileManager = IFileManager::Get();
+	auto& PlatformFile = FPlatformFileManager::Get().GetPlatformFile();
+
+	for (const auto& Pair : InDirtyFiles)
+	{
+		const auto& File = Pair.Key;
+		FString AbsoluteFilename = GetAbsolutePath(File.Get());
+		
+		const auto* CachedState = CachedDirectoryState.Files.Find(File);
+
+		const bool bFileExists = InFileSystemState ? InFileSystemState->Files.Find(File) != nullptr : PlatformFile.FileExists(*AbsoluteFilename);
+		if (bFileExists)
+		{
+			FFileData FileData;
+			if (const auto* FoundData = InFileSystemState ? InFileSystemState->Files.Find(File) : nullptr)
+			{
+				FileData = *FoundData;
+			}
+			else
+			{
+				// The dirty file timestamp is the time that the file was dirtied, not necessarily its modification time
+				FileData = FFileData(FileManager.GetTimeStamp(*AbsoluteFilename), Pair.Value.FileHash);
+			}
+
+			if (Config.bRequireFileHashes && !FileData.FileHash.IsValid())
+			{
+				// We don't have this file's hash yet. Temporarily ignore it.
+				continue;
+			}
+
+			// Do we think it exists in the cache?
+			if (CachedState)
+			{
+				// A file has changed if its hash is now different
+				if (Config.bRequireFileHashes &&
+					Config.ChangeDetectionBits[FFileCacheConfig::FileHash] &&
+					CachedState->FileHash != FileData.FileHash
+					)
+				{
+					ModifiedFiles.Add(File, FileData);
+				}
+				// or the timestamp has changed
+				else if (Config.ChangeDetectionBits[FFileCacheConfig::Timestamp] &&
+					CachedState->Timestamp != FileData.Timestamp)
+				{
+					ModifiedFiles.Add(File, FileData);
+				}
+				else
+				{
+					// File hasn't changed
+					InvalidDirtyFiles.Add(File);
+				}
+			}
+			else
+			{
+				AddedFiles.Add(File, FileData);
+			}
+		}
+		// We only report it as removed if it exists in the cache
+		else if (CachedState)
+		{
+			RemovedFiles.Add(File);
+		}
+		else
+		{
+			// File doesn't exist, and isn't in the cache
+			InvalidDirtyFiles.Add(File);
+		}
+	}
+
+	// Remove any dirty files that aren't dirty
+	for (auto& Filename : InvalidDirtyFiles)
+	{
+		InDirtyFiles.Remove(Filename);
+	}
+
+	// Rename / move detection
+	if (Config.bDetectMoves)
+	{
+		bool bHavePendingHashes = false;
+
+		// Remove any additions that don't have their hash generated yet
+		for (auto AdIt = AddedFiles.CreateIterator(); AdIt; ++AdIt)
+		{
+			if (!AdIt.Value().FileHash.IsValid())
+			{
+				bHavePendingHashes = true;
+				AdIt.RemoveCurrent();
+			}
+		}
+
+		// We can only detect renames or moves for files that have had their file hash harvested.
+		// If we can't find a valid move destination for this file, and we have pending hashes, ignore the removal until we can be sure it's not a move
+		for (auto RemoveIt = RemovedFiles.CreateIterator(); RemoveIt; ++RemoveIt)
+		{
+			const auto* CachedState = CachedDirectoryState.Files.Find(*RemoveIt);
+			if (CachedState && CachedState->FileHash.IsValid())
+			{
+				for (auto AdIt = AddedFiles.CreateIterator(); AdIt; ++AdIt)
+				{
+					if (AdIt.Value().FileHash == CachedState->FileHash)
+					{
+						// Found a move destination!
+						OutTransactions.Add(FUpdateCacheTransaction(*RemoveIt, AdIt.Key(), AdIt.Value()));
+
+						AdIt.RemoveCurrent();
+						RemoveIt.RemoveCurrent();
+						goto next;
+					}
+				}
+
+				// We can't be sure this isn't a move (yet) so temporarily ignore this
+				if (bHavePendingHashes)
+				{
+					RemoveIt.RemoveCurrent();
+				}
+			}
+
+		next:
+			continue;
+		}
+	}
+
+	for (auto& RemovedFile : RemovedFiles)
+	{
+		OutTransactions.Add(FUpdateCacheTransaction(MoveTemp(RemovedFile), EFileAction::Removed));
+	}
+	// RemovedFiles is now bogus
+
+	for (auto& Pair : AddedFiles)
+	{
+		OutTransactions.Add(FUpdateCacheTransaction(MoveTemp(Pair.Key), EFileAction::Added, Pair.Value));
+	}
+	// AddedFiles is now bogus
+	
+	for (auto& Pair : ModifiedFiles)
+	{
+		OutTransactions.Add(FUpdateCacheTransaction(MoveTemp(Pair.Key), EFileAction::Modified, Pair.Value));
+	}
+	// ModifiedFiles is now bogus
+}
+
+TArray<FUpdateCacheTransaction> FFileCache::GetOutstandingChanges()
+{
+	HarvestDirtyFileHashes();
+
+	TArray<FUpdateCacheTransaction> PendingTransactions;
+	DiffDirtyFiles(DirtyFiles, PendingTransactions);
+	DirtyFiles.Empty();
+	return PendingTransactions;
+}
+
+TArray<FUpdateCacheTransaction> FFileCache::FilterOutstandingChanges(const TFunctionRef<bool(const FUpdateCacheTransaction&, const FDateTime&)>& InPredicate)
+{
+	HarvestDirtyFileHashes();
+
+	TArray<FUpdateCacheTransaction> AllTransactions;
+	DiffDirtyFiles(DirtyFiles, AllTransactions, nullptr);
+
+	// Filter the transactions based on the predicate
+	TArray<FUpdateCacheTransaction> FilteredTransactions;
+	for (auto& Transaction : AllTransactions)
+	{
+		FFileData FileData = DirtyFiles.FindRef(Transaction.Filename);
+
+		// Timestamp is the time the file was dirtied, not necessarily the timestamp of the file
+		if (InPredicate(Transaction, FileData.Timestamp))
+		{
+			DirtyFiles.Remove(Transaction.Filename);
+			FilteredTransactions.Add(MoveTemp(Transaction));
+		}
+	}
+	// Anything left in AllTransactions is discarded
+	return FilteredTransactions;
+}
+
+void FFileCache::IgnoreNewFile(const FString& Filename)
+{
+	auto TransactionPath = GetTransactionPath(Filename);
+	if (TransactionPath.IsSet())
+	{
+		DirtyFiles.Remove(TransactionPath.GetValue());
+		
+		const FFileData FileData(IFileManager::Get().GetTimeStamp(*Filename), FMD5Hash::HashFile(*Filename));
+		CompleteTransaction(FUpdateCacheTransaction(MoveTemp(TransactionPath.GetValue()), EFileAction::Added, FileData));	
+	}
+}
+
+void FFileCache::IgnoreFileModification(const FString& Filename)
+{
+	auto TransactionPath = GetTransactionPath(Filename);
+	if (TransactionPath.IsSet())
+	{
+		DirtyFiles.Remove(TransactionPath.GetValue());
+		
+		const FFileData FileData(IFileManager::Get().GetTimeStamp(*Filename), FMD5Hash::HashFile(*Filename));
+		CompleteTransaction(FUpdateCacheTransaction(MoveTemp(TransactionPath.GetValue()), EFileAction::Modified, FileData));	
+	}
+}
+
+void FFileCache::IgnoreMovedFile(const FString& SrcFilename, const FString& DstFilename)
+{
+	auto SrcTransactionPath = GetTransactionPath(SrcFilename);
+	auto DstTransactionPath = GetTransactionPath(DstFilename);
+
+	if (SrcTransactionPath.IsSet() && DstTransactionPath.IsSet())
+	{
+		DirtyFiles.Remove(SrcTransactionPath.GetValue());
+		DirtyFiles.Remove(DstTransactionPath.GetValue());
+
+		const FFileData FileData(IFileManager::Get().GetTimeStamp(*DstFilename), FMD5Hash::HashFile(*DstFilename));
+		CompleteTransaction(FUpdateCacheTransaction(MoveTemp(SrcTransactionPath.GetValue()), MoveTemp(DstTransactionPath.GetValue()), FileData));	
+	}
+}
+
+void FFileCache::IgnoreDeletedFile(const FString& Filename)
+{
+	auto TransactionPath = GetTransactionPath(Filename);
+	if (TransactionPath.IsSet())
+	{
+		DirtyFiles.Remove(TransactionPath.GetValue());
+		CompleteTransaction(FUpdateCacheTransaction(MoveTemp(TransactionPath.GetValue()), EFileAction::Removed));	
+	}
+}
+
+void FFileCache::CompleteTransaction(FUpdateCacheTransaction&& Transaction)
+{
+	auto* CachedData = CachedDirectoryState.Files.Find(Transaction.Filename);
+	switch (Transaction.Action)
+	{
+	case EFileAction::Moved:
+		{
+			CachedDirectoryState.Files.Remove(Transaction.MovedFromFilename);
+			if (!CachedData)
+			{
+				CachedDirectoryState.Files.Add(Transaction.Filename, Transaction.FileData);
+			}
+			else
+			{
+				*CachedData = Transaction.FileData;
+			}
+
+			bSavedCacheDirty = true;
+		}
+		break;
+	case EFileAction::Modified:
+		if (CachedData)
+		{
+			// Update the timestamp
+			*CachedData = Transaction.FileData;
+
+			bSavedCacheDirty = true;
+		}
+		break;
+	case EFileAction::Added:
+		if (!CachedData)
+		{
+			// Add the file information to the cache
+			CachedDirectoryState.Files.Emplace(Transaction.Filename, Transaction.FileData);
+
+			bSavedCacheDirty = true;
+		}
+		break;
+	case EFileAction::Removed:
+		if (CachedData)
+		{
+			// Remove the file information to the cache
+			CachedDirectoryState.Files.Remove(Transaction.Filename);
+
+			bSavedCacheDirty = true;
+		}
+		break;
+
+	default:
+		checkf(false, TEXT("Invalid file cached transaction"));
+		break;
+	}
+}
+
+void FFileCache::Tick()
+{
+	/** Stage one: wait for the asynchronous directory reader to finish harvesting timestamps for the directory */
+	if (DirectoryReader.IsValid())
+	{
+		if (!DirectoryReader->IsComplete())
+		{
+			return;
+		}
+		else
+		{
+			ReadStateFromAsyncReader();
+
+			if (Config.bRequireFileHashes)
+			{
+				auto FilesThatNeedHashing = DirectoryReader->GetFilesThatNeedHashing();
+				if (FilesThatNeedHashing.Num() > 0)
+				{
+					AsyncFileHasher = MakeShareable(new FAsyncFileHasher(MoveTemp(FilesThatNeedHashing)));
+
+					AsyncTaskThread.AddTask(AsyncFileHasher);
+				}
+			}
+
+			// Null out our pointer to the directory reader to indicate that we've finished
+			DirectoryReader = nullptr;
+		}
+	}
+	/** The file cache is now running, and will report changes. */
+	/** Keep harvesting file hashes from the file hashing task until complete. These are much slower to gather, and only required for rename/move detection. */
+	else if (AsyncFileHasher.IsValid())
+	{
+		double Now = FPlatformTime::Seconds();
+
+		if (Now - LastFileHashGetTime > 5.f)
+		{
+			LastFileHashGetTime = Now;
+			auto Hashes = AsyncFileHasher->GetCompletedData();
+			if (Hashes.Num() > 0)
+			{
+				bSavedCacheDirty = true;
+				for (const auto& Data : Hashes)
+				{
+					FImmutableString CachePath = (Config.PathType == EPathType::Relative) ? *Data.AbsoluteFilename + Config.Directory.Len() : *Data.AbsoluteFilename;
+
+					auto* FileData = CachedDirectoryState.Files.Find(CachePath);
+					if (FileData && !FileData->FileHash.IsValid())
+					{
+						FileData->FileHash = Data.FileHash;
+					}
+				}
+			}
+
+			if (AsyncFileHasher->IsComplete())
+			{
+				UE_LOG(LogFileCache, Log, TEXT("Retrieving MD5 hashes for directory '%s' took %.2fs"), *Config.Directory, AsyncFileHasher->GetAge());
+				AsyncFileHasher = nullptr;
+			}
+		}
+	}
+}
+
+void FFileCache::ReadStateFromAsyncReader()
+{
+	// We should only ever get here once. The directory reader has finished scanning, and we can now diff the results with what we had saved in the cache file.
+	check(DirectoryReader->IsComplete());
+	
+	TOptional<FDirectoryState> LiveState = DirectoryReader->GetLiveState();
+	TOptional<FDirectoryState> CachedState = DirectoryReader->GetCachedState();
+
+	if (!CachedState.IsSet() || !Config.bDetectChangesSinceLastRun)
+	{
+		// If we don't have any cached data yet, just use the file data we just harvested
+		CachedDirectoryState = MoveTemp(LiveState.GetValue());
+		bSavedCacheDirty = true;
+		return;
+	}
+	else
+	{
+		// Use the cache that we gave to the directory reader
+		CachedDirectoryState = MoveTemp(CachedState.GetValue());
+	}
+
+	const FDateTime Now = FDateTime::UtcNow();
+	// We already have cached data so we need to compare it with the harvested data
+	// to detect additions, modifications, and removals
+	for (const auto& FilenameAndData : LiveState->Files)
+	{
+		const FString& Filename = FilenameAndData.Key.Get();
+
+		// If the file we've discovered was not applicable to the old cache, we can't report a change for it as we don't know if it's new or not, just add it straight to the cache.
+		if (!CachedDirectoryState.Rules.IsFileApplicable(*Filename))
+		{
+			CachedDirectoryState.Files.Add(FilenameAndData.Key, FilenameAndData.Value);
+			bSavedCacheDirty = true;
+		}
+		else
+		{
+			const auto* CachedData = CachedDirectoryState.Files.Find(Filename);
+			if (!CachedData || CachedData->Timestamp != FilenameAndData.Value.Timestamp)
+			{
+				DirtyFiles.Add(FilenameAndData.Key, FFileData(Now, FMD5Hash()));
+			}
+		}
+	}
+
+	// Check for anything that doesn't exist on disk anymore
+	for (auto It = CachedDirectoryState.Files.CreateIterator(); It; ++It)
+	{
+		const FImmutableString& Filename = It.Key();
+		if (LiveState->Rules.IsFileApplicable(*Filename.Get()) && !LiveState->Files.Contains(Filename))
+		{
+			DirtyFiles.Add(Filename, FFileData(Now, FMD5Hash()));
+		}
+	}
+
+	RescanForDirtyFileHashes();
+
+	// Update the applicable extensions now that we've updated the cache
+	CachedDirectoryState.Rules = LiveState->Rules;
+}
+
+void FFileCache::HarvestDirtyFileHashes()
+{
+	if (!DirtyFileHasher.IsValid())
+	{
+		return;
+	}
+	else for (FFilenameAndHash& Data : DirtyFileHasher->GetCompletedData())
+	{
+		FImmutableString CachePath = (Config.PathType == EPathType::Relative) ? *Data.AbsoluteFilename + Config.Directory.Len() : *Data.AbsoluteFilename;
+
+		if (auto* FileData = DirtyFiles.Find(CachePath))
+		{
+			FileData->FileHash = Data.FileHash;
+		}
+	}
+
+	if (DirtyFileHasher->IsComplete())
+	{
+		DirtyFileHasher = nullptr;
+	}
+}
+
+void FFileCache::RescanForDirtyFileHashes()
+{
+	TArray<FFilenameAndHash> FilesThatNeedHashing;
+
+	for (const auto& Pair : DirtyFiles)
+	{
+		if (!Pair.Value.FileHash.IsValid())
+		{
+			FilesThatNeedHashing.Emplace(GetAbsolutePath(Pair.Key.Get()));
+		}
+	}
+
+	if (FilesThatNeedHashing.Num() > 0)
+	{
+		// Re-create the dirty file hasher with the new data that needs hashing. The old task will clean itself up if it already exists.
+		DirtyFileHasher = MakeShareable(new FAsyncFileHasher(MoveTemp(FilesThatNeedHashing)));
+		AsyncTaskThread.AddTask(DirtyFileHasher);
+	}
+}
+
+void FFileCache::OnDirectoryChanged(const TArray<FFileChangeData>& FileChanges)
+{
+	// Harvest any completed data from the file hasher before we discard it
+	HarvestDirtyFileHashes();
+
+	const FDateTime Now = FDateTime::UtcNow();
+	for (const auto& ThisEntry : FileChanges)
+	{
+		auto TransactionPath = GetTransactionPath(ThisEntry.Filename);
+		if (TransactionPath.IsSet())
+		{
+			// Add the file that changed to the dirty files map, potentially invalidating the MD5 hash (we'll need to calculate it again)
+			DirtyFiles.Add(MoveTemp(TransactionPath.GetValue()), FFileData(Now, FMD5Hash()));
+		}
+	}
+
+	RescanForDirtyFileHashes();
+}
+
+} // namespace DirectoryWatcher
diff --git a/Engine/Source/Developer/DirectoryWatcher/Private/FileCacheUtilities.cpp b/Engine/Source/Developer/DirectoryWatcher/Private/FileCacheUtilities.cpp
new file mode 100644
--- /dev/null
+++ b/Engine/Source/Developer/DirectoryWatcher/Private/FileCacheUtilities.cpp
@@ -0,0 +1,187 @@
+// Copyright 1998-2015 Epic Games, Inc. All Rights Reserved.
+
+#include "DirectoryWatcherPrivatePCH.h"
+#include "FileCacheUtilities.h"
+
+bool MatchExtensionString(const TCHAR* Filename, const TCHAR* Extensions)
+{
+	const int32 StrLength = FCString::Strlen(Filename);
+
+	const TCHAR* Ext = FCString::Strrchr(Filename, '.');
+	if (Ext && *(++Ext) != '\0')
+	{
+		const int32 ExtLength = StrLength - (Ext - Filename);
+
+		const TCHAR* Search = FCString::Strchr(Extensions, ';');
+		while(Search)
+		{
+			if (*(++Search) == '\0')
+			{
+				break;
+			}
+
+			if (FCString::Strnicmp(Search, Ext, ExtLength) == 0 && *(Search + ExtLength) == ';')
+			{
+				return true;
+			}
+
+			Search = FCString::Strchr(Search, ';');
+		}
+	}
+
+	return false;
+}
+
+/** Test MatchExtensionString */
+IMPLEMENT_SIMPLE_AUTOMATION_TEST(FMatchExtensionStringTest, "System.Plugins.Directory Watcher.File Cache.Extension Matching", EAutomationTestFlags::ATF_Editor)
+bool FMatchExtensionStringTest::RunTest(const FString& Parameters)
+{
+	auto Test = [](const TCHAR* Needle, const TCHAR* Haystack, bool bExpected) -> bool{
+
+		if (MatchExtensionString(Needle, Haystack) != bExpected)
+		{
+			if (bExpected)
+			{
+				UE_LOG(LogFileCache, Error, TEXT("Unable to match '%s' in '%s'"), Needle, Haystack);
+			}
+			else
+			{
+				UE_LOG(LogFileCache, Error, TEXT("Erroneously matched '%s' in '%s'"), Needle, Haystack);
+			}
+			return false;
+		}
+		return true;
+	};
+
+	bool Result = true;
+	Result = Test(TEXT("bla.txt"), TEXT(";;"), false) && Result;
+	Result = Test(TEXT("bla.txt"), TEXT(";"), false) && Result;
+	Result = Test(TEXT("bla.txt"), TEXT(""), false) && Result;
+	Result = Test(TEXT("bla.txt"), TEXT(";txt;"), true) && Result;
+	Result = Test(TEXT("bla.text"), TEXT(";txt;"), false) && Result;
+	Result = Test(TEXT("bla.txt1"), TEXT(";txt;"), false) && Result;
+	Result = Test(TEXT("bla."), TEXT(";bla;"), false) && Result;
+	Result = Test(TEXT("bla.png"), TEXT(";png;txt;"), true) && Result;
+	Result = Test(TEXT("bla.txt"), TEXT(";png;txt;"), true) && Result;
+	Result = Test(TEXT("/folder.bin/bla.txt"), TEXT(";png;txt;"), true) && Result;
+	Result = Test(TEXT("/folder.bin/bla"), TEXT(";png;bin;"), false) && Result;
+
+	return true;
+}
+
+namespace DirectoryWatcher
+{
+
+struct FWildcardRule : IMatchRule
+{
+	FWildcardString WildcardString;
+	bool bInclude;
+
+	virtual ~FWildcardRule(){}
+	
+	virtual TOptional<bool> IsFileApplicable(const TCHAR* Filename) const override
+	{
+		if (WildcardString.IsMatch(Filename))
+		{
+			return bInclude;
+		}
+		return TOptional<bool>();
+	}
+
+	virtual void Serialize(FArchive& Ar) override
+	{
+		Ar << WildcardString;
+		Ar << bInclude;
+	}
+};
+
+void FMatchRules::AddWildcardRule(const FWildcardString& WildcardString, bool bInclude)
+{
+	FMatchRule Rule;
+	Rule.Type = FMatchRule::Wildcard;
+	auto* Impl = new FWildcardRule;
+	Impl->WildcardString = WildcardString;
+	Impl->bInclude = bInclude;
+	Rule.RuleImpl = MakeShareable(Impl);
+	Impls.Add(Rule);
+
+	// If there are any include patterns, we default to not matching all files
+	if (bInclude)
+	{
+		bDefaultIncludeState = false;
+	}
+}
+
+void FMatchRules::SetApplicableExtensions(const FString& InExtensions)
+{
+	ApplicableExtensions = InExtensions;
+
+	// Ensure that the extension strings are of the form ;ext1;ext2;ext3;
+	if (ApplicableExtensions[ApplicableExtensions.Len() - 1] != ';')
+	{
+		ApplicableExtensions += TEXT(";");
+	}
+	if (ApplicableExtensions[0] != ';')
+	{
+		ApplicableExtensions.InsertAt(0, ';');
+	}
+
+}
+
+void FMatchRules::FMatchRule::Serialize(FArchive& Ar)
+{
+	Ar << Type;
+	
+	if (Ar.IsLoading())
+	{
+		if (Type == Wildcard)
+		{
+			RuleImpl = MakeShareable(new FWildcardRule);
+		}
+	}
+
+	if (RuleImpl.IsValid())
+	{
+		RuleImpl->Serialize(Ar);
+	}
+}
+
+bool FMatchRules::IsFileApplicable(const TCHAR* Filename) const
+{
+	if (!ApplicableExtensions.IsEmpty() && !MatchExtensionString(Filename, *ApplicableExtensions))
+	{
+		return false;
+	}
+	
+	// If we have no rules, we match everything
+	if (Impls.Num() == 0)
+	{
+		return true;
+	}
+
+	bool bApplicable = bDefaultIncludeState;
+
+	// Otherwise we match any of the rules
+	for (const auto& Rule : Impls)
+	{
+		if (Rule.RuleImpl.IsValid())
+		{
+			auto IsApplicable = Rule.RuleImpl->IsFileApplicable(Filename);
+			if (IsApplicable.IsSet())
+			{
+				if (!IsApplicable.GetValue())
+				{
+					return false;
+				}
+				else
+				{
+					bApplicable = true;
+				}
+			}
+		}
+	}
+
+	return bApplicable;
+}
+
+} // namespace DirectoryWatcher
diff --git a/Engine/Source/Developer/DirectoryWatcher/Public/FileCache.h b/Engine/Source/Developer/DirectoryWatcher/Public/FileCache.h
new file mode 100644
--- /dev/null
+++ b/Engine/Source/Developer/DirectoryWatcher/Public/FileCache.h
@@ -0,0 +1,503 @@
+// Copyright 1998-2015 Epic Games, Inc. All Rights Reserved.
+
+#pragma once
+
+#include "Serialization/CustomVersion.h"
+#include "FileCacheUtilities.h"
+#include "SecureHash.h"
+#include "IDirectoryWatcher.h"
+
+namespace DirectoryWatcher
+{
+
+/** Custom serialization version for FFileCache */
+struct DIRECTORYWATCHER_API FFileCacheCustomVersion
+{
+	static const FGuid Key;
+	enum Type {	Initial, IncludeFileHash, Latest = IncludeFileHash };
+};
+
+/** Structure representing specific information about a particular file */
+struct FFileData
+{
+	/** Constructiors */
+	FFileData() : Timestamp(0) {}
+	FFileData(const FDateTime& InTimestamp, const FMD5Hash& InFileHash) : Timestamp(InTimestamp), FileHash(InFileHash) {}
+
+	friend bool operator==(const FFileData& LHS, const FFileData& RHS)
+	{
+		return LHS.Timestamp == RHS.Timestamp && LHS.FileHash == RHS.FileHash;
+	}
+
+	friend bool operator!=(const FFileData& LHS, const FFileData& RHS)
+	{
+		return LHS.Timestamp != RHS.Timestamp || LHS.FileHash != RHS.FileHash;	
+	}
+
+	/** Serializer for this type */
+	friend FArchive& operator<<(FArchive& Ar, FFileData& Data)
+	{
+		if (Ar.CustomVer(FFileCacheCustomVersion::Key) >= FFileCacheCustomVersion::Initial)
+		{
+			Ar << Data.Timestamp;
+		}
+
+		if (Ar.CustomVer(FFileCacheCustomVersion::Key) >= FFileCacheCustomVersion::IncludeFileHash)
+		{
+			Ar << Data.FileHash;
+		}
+
+		return Ar;
+	}
+
+	/** The cached timestamp of the file on disk */
+	FDateTime Timestamp;
+
+	/** The cached MD5 hash of the file on disk */
+	FMD5Hash FileHash;
+};
+
+/** Structure representing the file data for a number of files in a directory */
+struct FDirectoryState
+{
+	/** Default construction */
+	FDirectoryState() {}
+
+	/** Move construction */
+	FDirectoryState(FDirectoryState&& In) : Rules(MoveTemp(In.Rules)), Files(MoveTemp(In.Files)) {}
+	FDirectoryState& operator=(FDirectoryState&& In) { Swap(Rules, In.Rules); Swap(Files, In.Files); return *this; }
+
+	/** The rules that define what this state applies to */
+	FMatchRules Rules;
+
+	/** Filename -> data map */
+	TMap<FImmutableString, FFileData> Files;
+
+	/** Serializer for this type */
+	friend FArchive& operator<<(FArchive& Ar, FDirectoryState& State)
+	{
+		Ar.UsingCustomVersion(FFileCacheCustomVersion::Key);
+
+		// Ignore any old versions to data to ensure that we generate a new cache
+		if (Ar.CustomVer(FFileCacheCustomVersion::Key) >= FFileCacheCustomVersion::IncludeFileHash)
+		{
+			Ar << State.Rules;
+
+			// Number of files
+			int32 Num = State.Files.Num();
+			Ar << Num;
+			if (Ar.IsLoading())
+			{
+				State.Files.Reserve(Num);
+			}
+
+			Ar << State.Files;
+		}
+
+		return Ar;
+	}
+};
+
+enum class EFileAction : uint8
+{
+	Added, Modified, Removed, Moved
+};
+
+/** A transaction issued by FFileCache to describe a change to the cache. The change is only committed once the transaction is returned to the cache (see FFileCache::CompleteTransaction). */
+struct FUpdateCacheTransaction
+{
+	/** The path of the file to which this transaction relates */
+	FImmutableString Filename;
+
+	/** In the case of a moved file, this represents the path the file was moved from */
+	FImmutableString MovedFromFilename;
+
+	/** File data pertaining to this change at the time of dispatch */
+	FFileData FileData;
+
+	/** The type of action that prompted this transaction */
+	EFileAction Action;
+
+	/** Publically moveable */
+	FUpdateCacheTransaction(FUpdateCacheTransaction&& In) : Filename(MoveTemp(In.Filename)), MovedFromFilename(MoveTemp(In.MovedFromFilename)), FileData(MoveTemp(In.FileData)), Action(MoveTemp(In.Action)) {}
+	FUpdateCacheTransaction& operator=(FUpdateCacheTransaction&& In) { Swap(Filename, In.Filename); Swap(MovedFromFilename, In.MovedFromFilename); Swap(FileData, In.FileData); Swap(Action, In.Action); return *this; }
+
+private:
+	friend class FFileCache;
+
+	/** Construction responsibility is held by FFileCache */
+	FUpdateCacheTransaction(FImmutableString InFilename, EFileAction InAction, const FFileData& InFileData = FFileData())
+		: Filename(MoveTemp(InFilename)), FileData(InFileData), Action(InAction)
+	{}
+
+	/** Construction responsibility is held by FFileCache */
+	FUpdateCacheTransaction(FImmutableString InMovedFromFilename, FImmutableString InMovedToFilename, const FFileData& InFileData)
+		: Filename(MoveTemp(InMovedToFilename)), MovedFromFilename(MoveTemp(InMovedFromFilename)), FileData(InFileData), Action(EFileAction::Moved)
+	{}
+	
+	/** Not Copyable */
+	FUpdateCacheTransaction(const FUpdateCacheTransaction&) = delete;
+	FUpdateCacheTransaction& operator=(const FUpdateCacheTransaction&) = delete;
+};
+
+/** Enum specifying whether a path should be relative or absolute */
+enum EPathType
+{
+	/** Paths should be cached relative to the root cache directory */
+	Relative,
+
+	/** Paths should be cached as absolute file system paths */
+	Absolute
+};
+
+struct IAsyncFileCacheTask : public TSharedFromThis<IAsyncFileCacheTask, ESPMode::ThreadSafe>
+{
+	enum class EProgressResult
+	{
+		Finished, Pending
+	};
+
+	IAsyncFileCacheTask() : StartTime(FPlatformTime::Seconds()) {}
+	virtual ~IAsyncFileCacheTask() {}
+	
+	/** Tick this task. Only to be called on the task thread. */
+	virtual EProgressResult Tick(const FTimeLimit& TimeLimit) = 0;
+
+	/** Check whether this task is complete. Must be implemented in a thread-safe manner. */
+	virtual bool IsComplete() const = 0;
+
+	/** Get the age of this task in seconds */
+	double GetAge() const { return FPlatformTime::Seconds() - StartTime; }
+
+protected:
+
+	/** The time this task started */
+	double StartTime;
+};
+
+/** Simple struct that encapsulates a filename and its associated MD5 hash */
+struct FFilenameAndHash
+{
+	FString AbsoluteFilename;
+	FMD5Hash FileHash;
+
+	FFilenameAndHash(){}
+	FFilenameAndHash(const FString& File) : AbsoluteFilename(File) {}
+};
+
+/** Async task responsible for MD5 hashing a number of files, reporting completed hashes to the client when done */
+struct FAsyncFileHasher : public IAsyncFileCacheTask
+{
+	/** Constructor */
+	FAsyncFileHasher(TArray<FFilenameAndHash> InFilesThatNeedHashing);
+
+	/** Return any completed filenames and their corresponding hashes */
+	TArray<FFilenameAndHash> GetCompletedData();
+
+	/** Returns true when this task has finished hashing all its files */
+	virtual bool IsComplete() const override;
+
+protected:
+
+	/** Tick this reader (hashes as many files as possible in the time allowed). Returns progress state. */
+	virtual EProgressResult Tick(const FTimeLimit& Limit) override;
+
+	/** The array of data that we will process */
+	TArray<FFilenameAndHash> Data;
+
+	/** The number of items we have returned to the client. Only accessed from the main thread. */
+	int32 NumReturned;
+
+	/** The number of files that we have hashed on the task thread. Atomic - safe to access from any thread. */
+	FThreadSafeCounter CurrentIndex;
+
+	/** Scratch buffer used for reading in files */
+	TArray<uint8> ScratchBuffer;
+};
+
+/**
+ * Class responsible for 'asynchronously' scanning a folder for files and timestamps.
+ * Example usage:
+ *		FAsyncDirectoryReader Reader(TEXT("C:\\Path"), EPathType::Relative);
+ *
+ *		while(!Reader.IsComplete())
+ *		{
+ *			FPlatformProcess::Sleep(1);
+ *			Reader.Tick(FTimedSignal(1));	// Do 1 second of work
+ *		}
+ *		TOptional<FDirectoryState> State = Reader.GetFinalState();
+ */
+struct FAsyncDirectoryReader : public IAsyncFileCacheTask
+{
+	/** Constructor that sets up the directory reader to the specified directory */
+	FAsyncDirectoryReader(const FString& InDirectory, EPathType InPathType);
+
+	/** Set what files are relevant to this reader. Calling this once the reader starts results in undefined behaviour. */
+	void SetMatchRules(const FMatchRules& InRules)
+	{
+		if (LiveState.IsSet())
+		{
+			LiveState->Rules = InRules;
+		}
+	}
+
+	/**
+	 * Get the state of the directory once finished. Relinquishes the currently stored directory state to the client.
+	 * Returns nothing if incomplete, or if GetLiveState() has already been called.
+	 */
+	TOptional<FDirectoryState> GetLiveState();
+
+	/** Retrieve the cached state supplied to this class through UseCachedState(). */
+	TOptional<FDirectoryState> GetCachedState();
+
+	/** Retrieve the cached state supplied to this class through UseCachedState(). */
+	TArray<FFilenameAndHash> GetFilesThatNeedHashing()
+	{
+		TArray<FFilenameAndHash> Swapped;
+		Swap(FilesThatNeedHashing, Swapped);
+		return Swapped;
+	}
+
+	/** Instruct the directory reader to use the specified cached state to lookup file hashes, where timestamps haven't changed */
+	void UseCachedState(FDirectoryState InCachedState)
+	{
+		CachedState = MoveTemp(InCachedState);
+	}
+
+	/** Returns true when this directory reader has finished scanning the directory */
+	virtual bool IsComplete() const override;
+
+	/** Tick this reader (discover new directories / files). Returns progress state. */
+	virtual EProgressResult Tick(const FTimeLimit& Limit) override;
+
+private:
+	/** Non-recursively scan a single directory for its contents. Adds results to Pending arrays. */
+	void ScanDirectory(const FString& InDirectory);
+
+	/** Path to the root directory we want to scan */
+	FString RootPath;
+
+	/** Whether we should return relative or absolute paths */
+	EPathType PathType;
+
+	/** The currently discovered state of the directory - reset once relinquished to the client through GetLiveState */
+	TOptional<FDirectoryState> LiveState;
+
+	/** The previously cached  state of the directory, optional */
+	TOptional<FDirectoryState> CachedState;
+
+	/** An array of files that need hashing */
+	TArray<FFilenameAndHash> FilesThatNeedHashing;
+
+	/** A list of directories we have recursively found on our travels */
+	TArray<FString> PendingDirectories;
+
+	/** A list of files we have recursively found on our travels */
+	TArray<FString> PendingFiles;
+
+	/** Thread safe flag to signify when this class has finished reading */
+	FThreadSafeBool bIsComplete;
+};
+
+/** Configuration structure required to construct a FFileCache */
+struct FFileCacheConfig
+{
+	/** Enum that specifies what changes are required for a change to be reported. When combined, any valid change is reported. */
+	enum EChangeDetection
+	{
+		/** Report modifications when the timestamp of a file changes */
+		Timestamp,
+		/** Report modifications when the contents of a file changes */
+		FileHash,
+	};
+
+	FFileCacheConfig(FString InDirectory, FString InCacheFile)
+		: Directory(InDirectory), CacheFile(InCacheFile), PathType(EPathType::Relative), bDetectChangesSinceLastRun(false)
+		, ChangeDetectionBits(false, 2)
+	{
+		DetectMoves(true);
+		ChangeDetectionBits[EChangeDetection::Timestamp] = true;
+	}
+
+	/** String specifying the directory on disk that the cache should reflect */
+	FString Directory;
+
+	/** String specifying the file that the cache should be saved to. When empty, no cache file will be maintained (thus only an in-memory cache is used) */
+	FString CacheFile;
+
+	/** List of rules which define what we will be watching */
+	FMatchRules Rules;
+
+	/** Path type to return, relative to the directory or absolute. */
+	EPathType PathType;
+
+	/** When true, changes to the directory since the cache shutdown will be detected and reported. When false, said changes will silently be applied to the serialized cache. */
+	bool bDetectChangesSinceLastRun;
+
+	/** Set up this cache to detect moves */
+	FFileCacheConfig& DetectMoves(bool bInDetectMoves)
+	{
+		bDetectMoves = bInDetectMoves;
+		if (bDetectMoves)
+		{
+			bRequireFileHashes = true;
+		}
+		return *this;
+	}
+	
+	/** Set up this cache to generate MD5 hashes for its constituent files */
+	FFileCacheConfig& RequireFileHashes(bool bInRequireFileHashes)
+	{
+		if (ensureMsg(bInRequireFileHashes || !bDetectMoves, TEXT("Unable to disable file hashing when move detection is enabled")))
+		{
+			bRequireFileHashes = bInRequireFileHashes;
+		}
+		return *this;
+	}
+
+	/** Instruct the cache to report the specified changes to files */
+	FFileCacheConfig& DetectChangesFor(EChangeDetection ChangeType, bool Value)
+	{
+		ChangeDetectionBits[ChangeType] = Value;
+		return *this;
+	}
+
+private:
+
+	/** True to detect moves and renames (based on file hash), false otherwise. Implies bRequireFileHashes. */
+	bool bDetectMoves;
+
+	/** When true, the cache will also calculate MD5 hashes for files. When true, an additional thread will be launched on startup to harvest unknown MD5 hashes for the directory. */
+	bool bRequireFileHashes;
+
+	/** Bitfied specifying how we will be detecting changes. See EChangeDetection. */
+	TBitArray<> ChangeDetectionBits;
+
+	friend class FFileCache;
+};
+
+/**
+ * A class responsible for scanning a directory, and maintaining a cache of its state (files and timestamps).
+ * Changes in the cache can be retrieved through GetOutstandingChanges(). Changes will be reported for any
+ * change in the cached state even between runs of the process.
+ */
+class DIRECTORYWATCHER_API FFileCache
+{
+public:
+
+	/** Construction from a config */
+	FFileCache(const FFileCacheConfig& InConfig);
+
+	FFileCache(const FFileCache&) = delete;
+	FFileCache& operator=(const FFileCache&) = delete;
+
+	/** Destructor */
+	~FFileCache();
+
+	/** Destroy this cache. Cleans out the in-memory state and deletes the cache file, if present. */
+	void Destroy();
+
+	/** Get the absolute path of the directory this cache reflects */
+	const FString& GetDirectory() const { return Config.Directory; }
+
+	/** Check whether this file cache has finished starting up yet. Does not imply rename/move detection is fully initialized. (see MoveDetectionInitialized()) */
+	bool HasStartedUp() const;
+
+	/** Check whether this move/rename detection has been initiated or not. This can take much longer than startup, so can be checked separately */
+	bool MoveDetectionInitialized() const;
+
+	/** Attempt to locate file data pertaining to the specified filename.
+	 *	@param		InFilename		The filename to find data for - either relative to the directory, or absolute, depending on Config.PathType.
+	 *	@return		The current cached file data pertaining to the specified filename, or nullptr. May not be completely up to date if there are outstanding transactions.
+	 */
+	const FFileData* FindFileData(FImmutableString InFilename) const;
+
+	/** Tick this FileCache */
+	void Tick();
+
+	/** Write out the cached file, if we have any changes to write */
+	void WriteCache();
+
+	/**
+	 * Return a transaction to the cache for completion. Will update the cached state with the change
+	 * described in the transaction, and mark the cache as needing to be saved.
+	 */
+	void CompleteTransaction(FUpdateCacheTransaction&& Transaction);
+
+	/** Report an external change to the manager, such that a subsequent equal change reported by the os be ignored */
+	void IgnoreNewFile(const FString& Filename);
+	void IgnoreFileModification(const FString& Filename);
+	void IgnoreMovedFile(const FString& SrcFilename, const FString& DstFilename);
+	void IgnoreDeletedFile(const FString& Filename);
+
+	/** Get the number of pending changes to the cache. */
+	int32 GetNumOutstandingChanges() const { return DirtyFiles.Num(); }
+
+	/** Get pending changes to the cache. Transactions must be returned to CompleteTransaction to update the cache.
+	 *  Filter predicate recieves a transaction and the time the change was reported. */
+	TArray<FUpdateCacheTransaction> FilterOutstandingChanges(const TFunctionRef<bool(const FUpdateCacheTransaction&, const FDateTime&)>& InPredicate);
+	TArray<FUpdateCacheTransaction> GetOutstandingChanges();
+
+private:
+
+	/** Called when the directory we are monitoring has been changed in some way */
+	void OnDirectoryChanged(const TArray<FFileChangeData>& FileChanges);
+
+	/** Diff the specified set of dirty files (absolute paths, or relative to the monitor directory), adding transactions to the specified array if necessary.
+	 *	Optionally takes a directory state from which we can retrieve current file system state, without having to ask the FS directly.
+	 *  Optionally exclude files that have changed since the specified threshold, to ensure that related events get grouped together correctly.
+	 */
+	void DiffDirtyFiles(TMap<FImmutableString, FFileData>& InDirtyFiles, TArray<FUpdateCacheTransaction>& OutTransactions, const FDirectoryState* InFileSystemState = nullptr) const;
+
+	/** Detect a rename for the specified file */
+	void DetectRename();
+
+	/** Get the absolute path from a transaction filename */
+	FString GetAbsolutePath(const FString& InTransactionPath) const;
+
+	/** Get the transaction path from an absolute filename. Returns nothing if the absolute path is invalid (not under the correct folder, or not applicable) */
+	TOptional<FString> GetTransactionPath(const FString& InAbsolutePath) const;
+
+	/** Unbind the watcher from the directory. Called on destruction. */
+	void UnbindWatcher();
+
+	/** Read the cache file data and return the contents */
+	TOptional<FDirectoryState> ReadCache() const;
+
+	/** Called when the initial async reader has finished harvesting file system timestamps */
+	void ReadStateFromAsyncReader();
+
+	/** Called to harvest any file hashes that have been generated by the DirtyFileHasher thread task */
+	void HarvestDirtyFileHashes();
+	void RescanForDirtyFileHashes();
+
+private:
+
+	/** Configuration settings applied on construction */
+	FFileCacheConfig Config;
+
+	/** Asynchronous directory reader responsible for gathering all file/timestamp information recursively from our cache directory */
+	TSharedPtr<FAsyncDirectoryReader, ESPMode::ThreadSafe> DirectoryReader;
+	/** Asynchronous task used to harvest the MD5 hashes of a set of filenames */
+	TSharedPtr<FAsyncFileHasher, ESPMode::ThreadSafe> AsyncFileHasher;
+	/** Asynchronous task used to harvest the MD5 hashes of a set of recently changed filenames */
+	TSharedPtr<FAsyncFileHasher, ESPMode::ThreadSafe> DirtyFileHasher;
+
+	/** A map of dirty files that we will use to report changes to the user. Timestamp value of FFileData here pertains to the *time of the change* not the timestamp of the file. */
+	TMap<FImmutableString, FFileData> DirtyFiles;
+
+	/** Our in-memory view of the cached directory state. */
+	FDirectoryState CachedDirectoryState;
+
+	/** Handle to the directory watcher delegate so we can delete it properly */
+	FDelegateHandle WatcherDelegate;
+
+	/** True when the cached state we have in memory is more up to date than the serialized file. Enables WriteCache() when true. */
+	bool bSavedCacheDirty;
+
+	/** The time we last retrieved file hashes from the thread */
+	double LastFileHashGetTime;
+
+};
+
+} // namespace DirectoryWatcher
diff --git a/Engine/Source/Developer/DirectoryWatcher/Public/FileCacheUtilities.h b/Engine/Source/Developer/DirectoryWatcher/Public/FileCacheUtilities.h
new file mode 100644
--- /dev/null
+++ b/Engine/Source/Developer/DirectoryWatcher/Public/FileCacheUtilities.h
@@ -0,0 +1,155 @@
+// Copyright 1998-2015 Epic Games, Inc. All Rights Reserved.
+
+#pragma once
+
+DEFINE_LOG_CATEGORY_STATIC(LogFileCache, Log, All);
+
+namespace DirectoryWatcher
+{
+
+/** An immutable string with a cached CRC for efficient comparison with other strings */
+struct FImmutableString
+{
+	/** Constructible from a raw string */
+	FImmutableString(FString InString = FString()) : String(MoveTemp(InString)), CachedHash(0) {}
+	FImmutableString& operator=(FString InString) { String = MoveTemp(InString); CachedHash = 0; return *this; }
+
+	FImmutableString(const TCHAR* InString) : String(InString), CachedHash(0) {}
+	FImmutableString& operator=(const TCHAR* InString) { String = InString; CachedHash = 0; return *this; }
+
+	FImmutableString(const FImmutableString&) = default;
+	FImmutableString& operator=(const FImmutableString&) = default;
+
+	/** Move construction/assignment */
+	FImmutableString(FImmutableString&& In) : String(MoveTemp(In.String)), CachedHash(MoveTemp(In.CachedHash)) {}
+	FImmutableString& operator=(FImmutableString&& In) { Swap(String, In.String); Swap(CachedHash, In.CachedHash); return *this; }
+
+	/** Check for equality with another immutable string (case-sensitive) */
+	friend bool operator==(const FImmutableString& A, const FImmutableString& B)
+	{
+		return GetTypeHash(A) == GetTypeHash(B) && A.String == B.String;
+	}
+
+	/** Calculate the type hash for this string */
+	friend uint32 GetTypeHash(const FImmutableString& In)
+	{
+		if (In.CachedHash == 0 && In.String.Len() != 0)
+		{
+			In.CachedHash = GetTypeHash(In.String);
+		}
+		return In.CachedHash;
+	}
+
+	/** Get the underlying string */
+	const FString& Get() const { return String; }
+
+	/** Serialise this string */
+	friend FArchive& operator<<(FArchive& Ar, FImmutableString& InString)
+	{
+		Ar << InString.String;
+		if (Ar.IsSaving())
+		{
+			GetTypeHash(InString);	
+		}
+		
+		Ar << InString.CachedHash;
+
+		return Ar;
+	}
+
+private:
+	/** The string itself. */
+	FString String;
+
+	/** The cached hash of our string. 0 until compared with another string. */
+	mutable uint32 CachedHash;
+};
+
+/** A time limit that counts down from the time of construction, until it hits a given delay */
+struct FTimeLimit
+{
+	/** Constructor specifying that we should never bail out early */
+	FTimeLimit() : Delay(-1) { Reset(); }
+
+	/** Constructor specifying not to run over the specified number of seconds */
+	FTimeLimit(float NumSeconds) : Delay(NumSeconds) { Reset(); }
+
+	/** Return true if this time limit has been set up, false otherwise */
+	bool IsValid() const { return Delay != -1; }
+	
+	/** Check whether we have exceeded the time limit */
+	bool Exceeded() const { return IsValid() && FPlatformTime::Seconds() >= StartTime + Delay; }
+
+	/** Reset the time limit to start timing again from the current time */
+	void Reset() { StartTime = FPlatformTime::Seconds(); }
+
+private:
+
+	/** The delay specified by the user */
+	float Delay;
+
+	/** The time we started */
+	double StartTime;
+};
+
+/** A rule that checks whether a file is applicable or not */
+struct IMatchRule
+{
+	/** Serialize this rule */
+	virtual void Serialize(FArchive& Ar) = 0;
+
+	/** Test to see if a file is applicable based on this rule. Returns true if so, false of not, or empty if the file doesn't match this rule. */
+	virtual TOptional<bool> IsFileApplicable(const TCHAR* Filename) const = 0;
+};
+
+/** A set of rules that specifies what files we'll be watching */
+struct FMatchRules
+{
+	FMatchRules() : bDefaultIncludeState(true) {}
+
+	/** Specify a wildcard match to include or exclude */
+	DIRECTORYWATCHER_API void AddWildcardRule(const FWildcardString& WildcardString, bool bInclude);
+
+	/** Specify a set of applicable extensions, ; separated. Provided as an optimization to early-out on files we're not interested in. */
+	DIRECTORYWATCHER_API void SetApplicableExtensions(const FString& InExtensions);
+
+	/** Check whether the specified file is applicable based on these rules or not */
+	DIRECTORYWATCHER_API bool IsFileApplicable(const TCHAR* Filename) const;
+
+	friend void operator<<(FArchive& Ar, FMatchRules& Rules)
+	{
+		Ar << Rules.bDefaultIncludeState;
+		Ar << Rules.ApplicableExtensions;
+		Ar << Rules.Impls;
+	}
+
+private:
+
+	/** Implementation of a match rule, wrapping up its type, and implementation */
+	struct FMatchRule
+	{
+		enum EType { Wildcard };
+		/** The type of this rule */
+		TEnumAsByte<EType> Type;
+		/** The runtime implementation */
+		TSharedPtr<IMatchRule> RuleImpl;
+
+	private:
+		void Serialize(FArchive& Ar);
+		friend void operator<<(FArchive& Ar, FMatchRule& Rule)
+		{
+			Rule.Serialize(Ar);
+		}
+	};
+
+	/** Optimization to ignore files that don't match the given set of extensions. Expected to be of the form ;ext1;ext2;ext3; */
+	FString ApplicableExtensions;
+
+	/** Array of implementations */
+	TArray<FMatchRule> Impls;
+
+	/** Default inclusion state. True when only exclude rules exist, false if there are any Include rules */
+	bool bDefaultIncludeState;
+};
+
+} // namespace DirectoryWatcher
diff --git a/Engine/Source/Editor/UnrealEd/Private/AutoReimport/AutoReimportManager.cpp b/Engine/Source/Editor/UnrealEd/Private/AutoReimport/AutoReimportManager.cpp
--- a/Engine/Source/Editor/UnrealEd/Private/AutoReimport/AutoReimportManager.cpp
+++ b/Engine/Source/Editor/UnrealEd/Private/AutoReimport/AutoReimportManager.cpp
@@ -24,7 +24,7 @@ enum class EStateMachineNode { CallOnce, CallMany };
 template<typename TState>
 struct FStateMachine
 {
-	typedef TFunction<TOptional<TState>(const FTimeLimit&)> FunctionType;
+	typedef TFunction<TOptional<TState>(const DirectoryWatcher::FTimeLimit&)> FunctionType;
 
 	/** Constructor that specifies the initial state of the machine */
 	FStateMachine(TState InitialState) : CurrentState(InitialState){}
@@ -42,7 +42,7 @@ struct FStateMachine
 	}
 
 	/** Tick this state machine with the given time limit. Will continuously enumerate the machine until TimeLimit is reached */
-	void Tick(const FTimeLimit& TimeLimit)
+	void Tick(const DirectoryWatcher::FTimeLimit& TimeLimit)
 	{
 		while (!TimeLimit.Exceeded())
 		{
@@ -155,13 +155,13 @@ private:
 	TOptional<ECurrentState> Idle();
 	
 	/** Process any remaining pending additions we have */
-	TOptional<ECurrentState> ProcessAdditions(const FTimeLimit& TimeLimit);
+	TOptional<ECurrentState> ProcessAdditions(const DirectoryWatcher::FTimeLimit& TimeLimit);
 
 	/** Save any packages that were created inside ProcessAdditions */
 	TOptional<ECurrentState> SavePackages();
 
 	/** Process any remaining pending modifications we have */
-	TOptional<ECurrentState> ProcessModifications(const FTimeLimit& TimeLimit);
+	TOptional<ECurrentState> ProcessModifications(const DirectoryWatcher::FTimeLimit& TimeLimit);
 
 	/** Process any remaining pending deletions we have */
 	TOptional<ECurrentState> ProcessDeletions();
@@ -193,7 +193,7 @@ private:
 	bool bGuardAssetChanges;
 
 	/** A timeout used to refresh directory monitors when the user has made an interactive change to the settings */
-	FTimeLimit ResetMonitorsTimeout;
+	DirectoryWatcher::FTimeLimit ResetMonitorsTimeout;
 
 	/** The paused state of the state machine */
 	ECurrentState PausedState;
@@ -241,13 +241,13 @@ FAutoReimportManager::FAutoReimportManager()
 		SetUpDirectoryMonitors();
 	}
 
-	StateMachine.Add(ECurrentState::Idle,					EStateMachineNode::CallOnce,	[this](const FTimeLimit& T){ return this->Idle();					});
-	StateMachine.Add(ECurrentState::Paused,					EStateMachineNode::CallOnce,	[this](const FTimeLimit& T){ return this->Paused();					});
-	StateMachine.Add(ECurrentState::Aborting,				EStateMachineNode::CallOnce,	[this](const FTimeLimit& T){ return this->Abort();					});
-	StateMachine.Add(ECurrentState::ProcessAdditions,		EStateMachineNode::CallMany,	[this](const FTimeLimit& T){ return this->ProcessAdditions(T);		});
-	StateMachine.Add(ECurrentState::ProcessModifications,	EStateMachineNode::CallMany,	[this](const FTimeLimit& T){ return this->ProcessModifications(T);	});
-	StateMachine.Add(ECurrentState::ProcessDeletions,		EStateMachineNode::CallMany,	[this](const FTimeLimit& T){ return this->ProcessDeletions();		});
-	StateMachine.Add(ECurrentState::SavePackages,			EStateMachineNode::CallOnce,	[this](const FTimeLimit& T){ return this->SavePackages();			});
+	StateMachine.Add(ECurrentState::Idle,					EStateMachineNode::CallOnce,	[this](const DirectoryWatcher::FTimeLimit& T){ return this->Idle();					});
+	StateMachine.Add(ECurrentState::Paused,					EStateMachineNode::CallOnce,	[this](const DirectoryWatcher::FTimeLimit& T){ return this->Paused();					});
+	StateMachine.Add(ECurrentState::Aborting,				EStateMachineNode::CallOnce,	[this](const DirectoryWatcher::FTimeLimit& T){ return this->Abort();					});
+	StateMachine.Add(ECurrentState::ProcessAdditions,		EStateMachineNode::CallMany,	[this](const DirectoryWatcher::FTimeLimit& T){ return this->ProcessAdditions(T);		});
+	StateMachine.Add(ECurrentState::ProcessModifications,	EStateMachineNode::CallMany,	[this](const DirectoryWatcher::FTimeLimit& T){ return this->ProcessModifications(T);	});
+	StateMachine.Add(ECurrentState::ProcessDeletions,		EStateMachineNode::CallMany,	[this](const DirectoryWatcher::FTimeLimit& T){ return this->ProcessDeletions();		});
+	StateMachine.Add(ECurrentState::SavePackages,			EStateMachineNode::CallOnce,	[this](const DirectoryWatcher::FTimeLimit& T){ return this->SavePackages();			});
 }
 
 TArray<FPathAndMountPoint> FAutoReimportManager::GetMonitoredDirectories() const
@@ -413,7 +413,7 @@ int32 FAutoReimportManager::GetNumUnprocessedChanges() const
 	}, 0);
 }
 
-TOptional<ECurrentState> FAutoReimportManager::ProcessAdditions(const FTimeLimit& TimeLimit)
+TOptional<ECurrentState> FAutoReimportManager::ProcessAdditions(const DirectoryWatcher::FTimeLimit& TimeLimit)
 {
 	auto NewState = HandlePauseAbort(ECurrentState::ProcessAdditions);
 	if (NewState.IsSet())
@@ -468,7 +468,7 @@ TOptional<ECurrentState> FAutoReimportManager::ProcessAdditions(const FTimeLimit
 	return ECurrentState::ProcessModifications;
 }
 
-TOptional<ECurrentState> FAutoReimportManager::ProcessModifications(const FTimeLimit& TimeLimit)
+TOptional<ECurrentState> FAutoReimportManager::ProcessModifications(const DirectoryWatcher::FTimeLimit& TimeLimit)
 {
 	auto NewState = HandlePauseAbort(ECurrentState::ProcessModifications);
 	if (NewState.IsSet())
@@ -607,7 +607,7 @@ TOptional<ECurrentState> FAutoReimportManager::Idle()
 			DirectoryMonitors.Empty();
 		}
 
-		ResetMonitorsTimeout = FTimeLimit();
+		ResetMonitorsTimeout = DirectoryWatcher::FTimeLimit();
 		return TOptional<ECurrentState>();
 	}
 
@@ -661,7 +661,7 @@ void FAutoReimportManager::Cleanup()
 void FAutoReimportManager::Tick(float DeltaTime)
 {
 	// Never spend more than a 60fps frame doing this work (meaning we shouldn't drop below 30fps), we can do more if we're throttling CPU usage (ie editor running in background)
-	const FTimeLimit TimeLimit(GEditor->ShouldThrottleCPUUsage() ? 1 / 6.f : 1.f / 60.f);
+	const DirectoryWatcher::FTimeLimit TimeLimit(GEditor->ShouldThrottleCPUUsage() ? 1 / 6.f : 1.f / 60.f);
 	StateMachine.Tick(TimeLimit);
 }
 
@@ -683,7 +683,7 @@ void FAutoReimportManager::HandleLoadingSavingSettingChanged(FName PropertyName)
 	if (PropertyName == GET_MEMBER_NAME_CHECKED(UEditorLoadingSavingSettings, bMonitorContentDirectories) ||
 		PropertyName == GET_MEMBER_NAME_CHECKED(UEditorLoadingSavingSettings, AutoReimportDirectorySettings))
 	{
-		ResetMonitorsTimeout = FTimeLimit(5.f);
+		ResetMonitorsTimeout = DirectoryWatcher::FTimeLimit(5.f);
 	}
 }
 
@@ -703,7 +703,7 @@ void FAutoReimportManager::SetUpDirectoryMonitors()
 	{
 		FString SourceDirectory;
 		FString MountPoint;
-		FMatchRules Rules;
+		DirectoryWatcher::FMatchRules Rules;
 	};
 
 	TArray<FParsedSettings> FinalArray;
diff --git a/Engine/Source/Editor/UnrealEd/Private/AutoReimport/AutoReimportTests.cpp b/Engine/Source/Editor/UnrealEd/Private/AutoReimport/AutoReimportTests.cpp
--- a/Engine/Source/Editor/UnrealEd/Private/AutoReimport/AutoReimportTests.cpp
+++ b/Engine/Source/Editor/UnrealEd/Private/AutoReimport/AutoReimportTests.cpp
@@ -30,14 +30,14 @@ private:
 };
 
 /** Generate a config from the specified options, to pass to FFileCache on construction */
-FFileCacheConfig GenerateFileCacheConfig(const FString& InWorkingDir)
+DirectoryWatcher::FFileCacheConfig GenerateFileCacheConfig(const FString& InWorkingDir)
 {
 	FString Directory = FPaths::ConvertRelativePathToFull(InWorkingDir);
 	FString CacheFilename = Directory / TEXT("Cache.bin");
 
-	FFileCacheConfig Config(Directory / TEXT("Content") / TEXT(""), MoveTemp(CacheFilename));
+	DirectoryWatcher::FFileCacheConfig Config(Directory / TEXT("Content") / TEXT(""), MoveTemp(CacheFilename));
 	// We always store paths inside content folders relative to the folder
-	Config.PathType = EPathType::Relative;
+	Config.PathType = DirectoryWatcher::EPathType::Relative;
 
 	Config.bDetectChangesSinceLastRun = false;
 	return Config;
@@ -46,8 +46,8 @@ FFileCacheConfig GenerateFileCacheConfig(const FString& InWorkingDir)
 /** Persistent test payload used for async testing, generally captured by lambdas */
 struct FAutoReimportTestPayload
 {
-	FFileCacheConfig Config;
-	TSharedPtr<FFileCache> FileCache;
+	DirectoryWatcher::FFileCacheConfig Config;
+	TSharedPtr<DirectoryWatcher::FFileCache> FileCache;
 	FString WorkingDir;
 
 	FAutoReimportTestPayload(const FString& InWorkingDir)
@@ -69,7 +69,7 @@ struct FAutoReimportTestPayload
 	{
 		if (!FileCache.IsValid())
 		{
-			FileCache = MakeShareable(new FFileCache(Config));
+			FileCache = MakeShareable(new DirectoryWatcher::FFileCache(Config));
 		}
 	}
 
@@ -164,7 +164,7 @@ bool FAutoReimportSimpleCreateTest::RunTest(const FString& Parameters)
 				UE_LOG(LogAutoReimportTests, Error, TEXT("Path reported incorrectly (%s != %s)."), *Changes[0].Filename.Get(), Filename);
 			}
 			// Check it was the correct type
-			else if (Changes[0].Action != EFileAction::Added)
+			else if (Changes[0].Action != DirectoryWatcher::EFileAction::Added)
 			{
 				UE_LOG(LogAutoReimportTests, Error, TEXT("Incorrect action for added file %s."), Filename);
 			}
@@ -233,7 +233,7 @@ bool FAutoReimportSimpleModifyTest::RunTest(const FString& Parameters)
 				UE_LOG(LogAutoReimportTests, Error, TEXT("Path reported incorrectly (%s != %s)."), *Changes[0].Filename.Get(), DstFilename);
 			}
 			// Check it was the correct type
-			else if (Changes[0].Action != EFileAction::Modified)
+			else if (Changes[0].Action != DirectoryWatcher::EFileAction::Modified)
 			{
 				UE_LOG(LogAutoReimportTests, Error, TEXT("Incorrect action for modified file %s."), DstFilename);
 			}
@@ -301,7 +301,7 @@ bool FAutoReimportSimpleDeleteTest::RunTest(const FString& Parameters)
 				UE_LOG(LogAutoReimportTests, Error, TEXT("Path reported incorrectly (%s != %s)."), *Changes[0].Filename.Get(), Filename);
 			}
 			// Check it was the correct type
-			else if (Changes[0].Action != EFileAction::Removed)
+			else if (Changes[0].Action != DirectoryWatcher::EFileAction::Removed)
 			{
 				UE_LOG(LogAutoReimportTests, Error, TEXT("Incorrect action for deleted file %s."), Filename);
 			}
@@ -358,7 +358,7 @@ bool FAutoReimportSimpleRenameTest::RunTest(const FString& Parameters)
 				UE_LOG(LogAutoReimportTests, Error, TEXT("Incorrect number of changes reported (%d != 1)."), Changes.Num());
 			}
 			// Check it was the correct type
-			else if (Changes[0].Action != EFileAction::Moved)
+			else if (Changes[0].Action != DirectoryWatcher::EFileAction::Moved)
 			{
 				UE_LOG(LogAutoReimportTests, Error, TEXT("Incorrect action for renamed file %s -> %s."), SrcFilename, DstFilename);
 			}
@@ -428,7 +428,7 @@ bool FAutoReimportSimpleMoveExternallyTest::RunTest(const FString& Parameters)
 				UE_LOG(LogAutoReimportTests, Error, TEXT("Incorrect number of changes reported (%d != 1)."), Changes.Num());
 			}
 			// Check it was the correct type
-			else if (Changes[0].Action != EFileAction::Removed)
+			else if (Changes[0].Action != DirectoryWatcher::EFileAction::Removed)
 			{
 				UE_LOG(LogAutoReimportTests, Error, TEXT("Incorrect action for removed file %s -> %s."), SrcFilename, DstFilename);
 			}
@@ -512,7 +512,7 @@ bool FAutoReimportRestartDetectionTest::RunTest(const FString& Parameters)
 				return;
 			}
 			// Check it was the correct type
-			else if (Changes[0].Action != EFileAction::Modified)
+			else if (Changes[0].Action != DirectoryWatcher::EFileAction::Modified)
 			{
 				UE_LOG(LogAutoReimportTests, Error, TEXT("Incorrect action for modified file %s."), SrcFilename);
 				return;
@@ -596,7 +596,7 @@ bool FAutoReimportMultipleChangesTest::RunTest(const FString& Parameters)
 				UE_LOG(LogAutoReimportTests, Error, TEXT("Incorrect number of changes reported (%d != 1)."), Changes.Num());
 			}
 			// Check it was the correct type
-			else if (Changes[0].Action != EFileAction::Added)
+			else if (Changes[0].Action != DirectoryWatcher::EFileAction::Added)
 			{
 				UE_LOG(LogAutoReimportTests, Error, TEXT("Incorrect action for added file %s."), Filename2);
 			}
@@ -645,7 +645,7 @@ bool FAutoReimportChangeExtensionsTest::RunTest(const FString& Parameters)
 				UE_LOG(LogAutoReimportTests, Error, TEXT("Incorrect number of changes reported (%d != 1)."), Changes.Num());
 				return;
 			}
-			else if (Changes[0].Action != EFileAction::Added)
+			else if (Changes[0].Action != DirectoryWatcher::EFileAction::Added)
 			{
 				UE_LOG(LogAutoReimportTests, Error, TEXT("Incorrect action for added file %s."), SrcFilename3);
 				return;
@@ -659,7 +659,7 @@ bool FAutoReimportChangeExtensionsTest::RunTest(const FString& Parameters)
 
 			Test->StopWatching();
 			// Change the extensions we're watching, and add another png
-			Test->Config.Rules = FMatchRules();
+			Test->Config.Rules = DirectoryWatcher::FMatchRules();
 			Test->Config.Rules.SetApplicableExtensions(TEXT("png;"));
 
 			Test->StartWatching();
diff --git a/Engine/Source/Editor/UnrealEd/Private/AutoReimport/AutoReimportUtilities.cpp b/Engine/Source/Editor/UnrealEd/Private/AutoReimport/AutoReimportUtilities.cpp
--- a/Engine/Source/Editor/UnrealEd/Private/AutoReimport/AutoReimportUtilities.cpp
+++ b/Engine/Source/Editor/UnrealEd/Private/AutoReimport/AutoReimportUtilities.cpp
@@ -7,184 +7,6 @@
 #include "IAssetRegistry.h"
 #include "AssetSourceFilenameCache.h"
 
-bool MatchExtensionString(const TCHAR* Filename, const TCHAR* Extensions)
-{
-	const int32 StrLength = FCString::Strlen(Filename);
-
-	const TCHAR* Ext = FCString::Strrchr(Filename, '.');
-	if (Ext && *(++Ext) != '\0')
-	{
-		const int32 ExtLength = StrLength - (Ext - Filename);
-
-		const TCHAR* Search = FCString::Strchr(Extensions, ';');
-		while(Search)
-		{
-			if (*(++Search) == '\0')
-			{
-				break;
-			}
-
-			if (FCString::Strnicmp(Search, Ext, ExtLength) == 0 && *(Search + ExtLength) == ';')
-			{
-				return true;
-			}
-
-			Search = FCString::Strchr(Search, ';');
-		}
-	}
-
-	return false;
-}
-
-/** Test MatchExtensionString */
-IMPLEMENT_SIMPLE_AUTOMATION_TEST(FMatchExtensionStringTest, "System.Editor.Auto Reimport.Extension Matching", EAutomationTestFlags::ATF_Editor)
-bool FMatchExtensionStringTest::RunTest(const FString& Parameters)
-{
-	auto Test = [](const TCHAR* Needle, const TCHAR* Haystack, bool bExpected) -> bool{
-
-		if (MatchExtensionString(Needle, Haystack) != bExpected)
-		{
-			if (bExpected)
-			{
-				UE_LOG(LogAutoReimportManager, Error, TEXT("Unable to match '%s' in '%s'"), Needle, Haystack);
-			}
-			else
-			{
-				UE_LOG(LogAutoReimportManager, Error, TEXT("Erroneously matched '%s' in '%s'"), Needle, Haystack);
-			}
-			return false;
-		}
-		return true;
-	};
-
-	bool Result = true;
-	Result = Test(TEXT("bla.txt"), TEXT(";;"), false) && Result;
-	Result = Test(TEXT("bla.txt"), TEXT(";"), false) && Result;
-	Result = Test(TEXT("bla.txt"), TEXT(""), false) && Result;
-	Result = Test(TEXT("bla.txt"), TEXT(";txt;"), true) && Result;
-	Result = Test(TEXT("bla.text"), TEXT(";txt;"), false) && Result;
-	Result = Test(TEXT("bla.txt1"), TEXT(";txt;"), false) && Result;
-	Result = Test(TEXT("bla."), TEXT(";bla;"), false) && Result;
-	Result = Test(TEXT("bla.png"), TEXT(";png;txt;"), true) && Result;
-	Result = Test(TEXT("bla.txt"), TEXT(";png;txt;"), true) && Result;
-	Result = Test(TEXT("/folder.bin/bla.txt"), TEXT(";png;txt;"), true) && Result;
-	Result = Test(TEXT("/folder.bin/bla"), TEXT(";png;bin;"), false) && Result;
-
-	return true;
-}
-
-struct FWildcardRule : IMatchRule
-{
-	FWildcardString WildcardString;
-	bool bInclude;
-
-	virtual ~FWildcardRule(){}
-	
-	virtual TOptional<bool> IsFileApplicable(const TCHAR* Filename) const override
-	{
-		if (WildcardString.IsMatch(Filename))
-		{
-			return bInclude;
-		}
-		return TOptional<bool>();
-	}
-
-	virtual void Serialize(FArchive& Ar) override
-	{
-		Ar << WildcardString;
-		Ar << bInclude;
-	}
-};
-
-void FMatchRules::AddWildcardRule(const FWildcardString& WildcardString, bool bInclude)
-{
-	FMatchRule Rule;
-	Rule.Type = FMatchRule::Wildcard;
-	auto* Impl = new FWildcardRule;
-	Impl->WildcardString = WildcardString;
-	Impl->bInclude = bInclude;
-	Rule.RuleImpl = MakeShareable(Impl);
-	Impls.Add(Rule);
-
-	// If there are any include patterns, we default to not matching all files
-	if (bInclude)
-	{
-		bDefaultIncludeState = false;
-	}
-}
-
-void FMatchRules::SetApplicableExtensions(const FString& InExtensions)
-{
-	ApplicableExtensions = InExtensions;
-
-	// Ensure that the extension strings are of the form ;ext1;ext2;ext3;
-	if (ApplicableExtensions[ApplicableExtensions.Len() - 1] != ';')
-	{
-		ApplicableExtensions += TEXT(";");
-	}
-	if (ApplicableExtensions[0] != ';')
-	{
-		ApplicableExtensions.InsertAt(0, ';');
-	}
-
-}
-
-void FMatchRules::FMatchRule::Serialize(FArchive& Ar)
-{
-	Ar << Type;
-	
-	if (Ar.IsLoading())
-	{
-		if (Type == Wildcard)
-		{
-			RuleImpl = MakeShareable(new FWildcardRule);
-		}
-	}
-
-	if (RuleImpl.IsValid())
-	{
-		RuleImpl->Serialize(Ar);
-	}
-}
-
-bool FMatchRules::IsFileApplicable(const TCHAR* Filename) const
-{
-	if (!ApplicableExtensions.IsEmpty() && !MatchExtensionString(Filename, *ApplicableExtensions))
-	{
-		return false;
-	}
-	
-	// If we have no rules, we match everything
-	if (Impls.Num() == 0)
-	{
-		return true;
-	}
-
-	bool bApplicable = bDefaultIncludeState;
-
-	// Otherwise we match any of the rules
-	for (const auto& Rule : Impls)
-	{
-		if (Rule.RuleImpl.IsValid())
-		{
-			auto IsApplicable = Rule.RuleImpl->IsFileApplicable(Filename);
-			if (IsApplicable.IsSet())
-			{
-				if (!IsApplicable.GetValue())
-				{
-					return false;
-				}
-				else
-				{
-					bApplicable = true;
-				}
-			}
-		}
-	}
-
-	return bApplicable;
-}
-
 namespace Utils
 {
 	TArray<FAssetData> FindAssetsPertainingToFile(const IAssetRegistry& Registry, const FString& AbsoluteFilename)
diff --git a/Engine/Source/Editor/UnrealEd/Private/AutoReimport/AutoReimportUtilities.h b/Engine/Source/Editor/UnrealEd/Private/AutoReimport/AutoReimportUtilities.h
--- a/Engine/Source/Editor/UnrealEd/Private/AutoReimport/AutoReimportUtilities.h
+++ b/Engine/Source/Editor/UnrealEd/Private/AutoReimport/AutoReimportUtilities.h
@@ -7,151 +7,6 @@ class IAssetRegistry;
 
 DEFINE_LOG_CATEGORY_STATIC(LogAutoReimportManager, Log, All);
 
-/** An immutable string with a cached CRC for efficient comparison with other strings */
-struct FImmutableString
-{
-	/** Constructible from a raw string */
-	FImmutableString(FString InString = FString()) : String(MoveTemp(InString)), CachedHash(0) {}
-	FImmutableString& operator=(FString InString) { String = MoveTemp(InString); CachedHash = 0; return *this; }
-
-	FImmutableString(const TCHAR* InString) : String(InString), CachedHash(0) {}
-	FImmutableString& operator=(const TCHAR* InString) { String = InString; CachedHash = 0; return *this; }
-
-	FImmutableString(const FImmutableString&) = default;
-	FImmutableString& operator=(const FImmutableString&) = default;
-
-	/** Move construction/assignment */
-	FImmutableString(FImmutableString&& In) : String(MoveTemp(In.String)), CachedHash(MoveTemp(In.CachedHash)) {}
-	FImmutableString& operator=(FImmutableString&& In) { Swap(String, In.String); Swap(CachedHash, In.CachedHash); return *this; }
-
-	/** Check for equality with another immutable string (case-sensitive) */
-	friend bool operator==(const FImmutableString& A, const FImmutableString& B)
-	{
-		return GetTypeHash(A) == GetTypeHash(B) && A.String == B.String;
-	}
-
-	/** Calculate the type hash for this string */
-	friend uint32 GetTypeHash(const FImmutableString& In)
-	{
-		if (In.CachedHash == 0 && In.String.Len() != 0)
-		{
-			In.CachedHash = GetTypeHash(In.String);
-		}
-		return In.CachedHash;
-	}
-
-	/** Get the underlying string */
-	const FString& Get() const { return String; }
-
-	/** Serialise this string */
-	friend FArchive& operator<<(FArchive& Ar, FImmutableString& InString)
-	{
-		Ar << InString.String;
-		if (Ar.IsSaving())
-		{
-			GetTypeHash(InString);	
-		}
-		
-		Ar << InString.CachedHash;
-
-		return Ar;
-	}
-
-private:
-	/** The string itself. */
-	FString String;
-
-	/** The cached hash of our string. 0 until compared with another string. */
-	mutable uint32 CachedHash;
-};
-
-/** A time limit that counts down from the time of construction, until it hits a given delay */
-struct FTimeLimit
-{
-	/** Constructor specifying that we should never bail out early */
-	FTimeLimit() : Delay(-1) { Reset(); }
-
-	/** Constructor specifying not to run over the specified number of seconds */
-	FTimeLimit(float NumSeconds) : Delay(NumSeconds) { Reset(); }
-
-	/** Return true if this time limit has been set up, false otherwise */
-	bool IsValid() const { return Delay != -1; }
-	
-	/** Check whether we have exceeded the time limit */
-	bool Exceeded() const { return IsValid() && FPlatformTime::Seconds() >= StartTime + Delay; }
-
-	/** Reset the time limit to start timing again from the current time */
-	void Reset() { StartTime = FPlatformTime::Seconds(); }
-
-private:
-
-	/** The delay specified by the user */
-	float Delay;
-
-	/** The time we started */
-	double StartTime;
-};
-
-/** A rule that checks whether a file is applicable or not */
-struct IMatchRule
-{
-	/** Serialize this rule */
-	virtual void Serialize(FArchive& Ar) = 0;
-
-	/** Test to see if a file is applicable based on this rule. Returns true if so, false of not, or empty if the file doesn't match this rule. */
-	virtual TOptional<bool> IsFileApplicable(const TCHAR* Filename) const = 0;
-};
-
-/** A set of rules that specifies what files apply to the auto reimporter */
-struct FMatchRules
-{
-	FMatchRules() : bDefaultIncludeState(true) {}
-
-	/** Specify a wildcard match to include or exclude */
-	void AddWildcardRule(const FWildcardString& WildcardString, bool bInclude);
-
-	/** Specify a set of applicable extensions, ; separated. Provided as an optimization to early-out on files we're not interested in. */
-	void SetApplicableExtensions(const FString& InExtensions);
-
-	/** Check whether the specified file is applicable based on these rules or not */
-	bool IsFileApplicable(const TCHAR* Filename) const;
-
-	friend void operator<<(FArchive& Ar, FMatchRules& Rules)
-	{
-		Ar << Rules.bDefaultIncludeState;
-		Ar << Rules.ApplicableExtensions;
-		Ar << Rules.Impls;
-	}
-
-private:
-
-	/** Implementation of a match rule, wrapping up its type, and implementation */
-	struct FMatchRule
-	{
-		enum EType { Wildcard };
-		/** The type of this rule */
-		TEnumAsByte<EType> Type;
-		/** The runtime implementation */
-		TSharedPtr<IMatchRule> RuleImpl;
-
-	private:
-		void Serialize(FArchive& Ar);
-		friend void operator<<(FArchive& Ar, FMatchRule& Rule)
-		{
-			Rule.Serialize(Ar);
-		}
-	};
-
-	/** Optimization to ignore files that don't match the given set of extensions. Expected to be of the form ;ext1;ext2;ext3; */
-	FString ApplicableExtensions;
-
-	/** Array of implementations */
-	TArray<FMatchRule> Impls;
-
-	/** Default inclusion state. True when only exclude rules exist, false if there are any Include rules */
-	bool bDefaultIncludeState;
-};
-
 namespace Utils
 {
 	/** Reduce the array to the specified accumulator */
@@ -171,4 +26,4 @@ namespace Utils
 	/** Extract any source file paths from the specified object */
 	TArray<FString> ExtractSourceFilePaths(UObject* Object);
 	void ExtractSourceFilePaths(UObject* Object, TArray<FString>& OutSourceFiles);
-}
\ No newline at end of file
+}
diff --git a/Engine/Source/Editor/UnrealEd/Private/AutoReimport/ContentDirectoryMonitor.cpp b/Engine/Source/Editor/UnrealEd/Private/AutoReimport/ContentDirectoryMonitor.cpp
--- a/Engine/Source/Editor/UnrealEd/Private/AutoReimport/ContentDirectoryMonitor.cpp
+++ b/Engine/Source/Editor/UnrealEd/Private/AutoReimport/ContentDirectoryMonitor.cpp
@@ -20,7 +20,7 @@ bool IsAssetDirty(UObject* Asset)
 }
 
 /** Generate a config from the specified options, to pass to FFileCache on construction */
-FFileCacheConfig GenerateFileCacheConfig(const FString& InPath, const FMatchRules& InMatchRules, const FString& InMountedContentPath)
+DirectoryWatcher::FFileCacheConfig GenerateFileCacheConfig(const FString& InPath, const DirectoryWatcher::FMatchRules& InMatchRules, const FString& InMountedContentPath)
 {
 	FString Directory = FPaths::ConvertRelativePathToFull(InPath);
 
@@ -28,23 +28,23 @@ FFileCacheConfig GenerateFileCacheConfig(const FString& InPath, const FMatchRule
 	const uint32 CRC = FCrc::MemCrc32(*HashString, HashString.Len()*sizeof(TCHAR));	
 	FString CacheFilename = FPaths::ConvertRelativePathToFull(FPaths::GameIntermediateDir()) / TEXT("ReimportCache") / FString::Printf(TEXT("%u.bin"), CRC);
 
-	FFileCacheConfig Config(MoveTemp(Directory), MoveTemp(CacheFilename));
+	DirectoryWatcher::FFileCacheConfig Config(MoveTemp(Directory), MoveTemp(CacheFilename));
 	Config.Rules = InMatchRules;
 	// We always store paths inside content folders relative to the folder
-	Config.PathType = EPathType::Relative;
+	Config.PathType = DirectoryWatcher::EPathType::Relative;
 
 	Config.bDetectChangesSinceLastRun = GetDefault<UEditorLoadingSavingSettings>()->bDetectChangesOnRestart;
 
 	// We only detect changes for when the file *contents* have changed (not its timestamp)
 	Config
 		.DetectMoves(true)
-		.DetectChangesFor(FFileCacheConfig::Timestamp, false)
-		.DetectChangesFor(FFileCacheConfig::FileHash, true);
+		.DetectChangesFor(DirectoryWatcher::FFileCacheConfig::Timestamp, false)
+		.DetectChangesFor(DirectoryWatcher::FFileCacheConfig::FileHash, true);
 
 	return Config;
 }
 
-FContentDirectoryMonitor::FContentDirectoryMonitor(const FString& InDirectory, const FMatchRules& InMatchRules, const FString& InMountedContentPath)
+FContentDirectoryMonitor::FContentDirectoryMonitor(const FString& InDirectory, const DirectoryWatcher::FMatchRules& InMatchRules, const FString& InMountedContentPath)
 	: Cache(GenerateFileCacheConfig(InDirectory, InMatchRules, InMountedContentPath))
 	, MountedContentPath(InMountedContentPath)
 	, LastSaveTime(0)
@@ -96,7 +96,7 @@ int32 FContentDirectoryMonitor::StartProcessing()
 	const FDateTime Threshold = FDateTime::UtcNow() - FTimespan(0, 0, GetDefault<UEditorLoadingSavingSettings>()->AutoReimportThreshold);
 
 	// Get all the changes that have happend beyond our import threshold
-	auto OutstandingChanges = Cache.FilterOutstandingChanges([&](const FUpdateCacheTransaction& Transaction, const FDateTime& TimeOfChange){
+	auto OutstandingChanges = Cache.FilterOutstandingChanges([&](const DirectoryWatcher::FUpdateCacheTransaction& Transaction, const FDateTime& TimeOfChange){
 		return TimeOfChange <= Threshold;
 	});
 
@@ -110,7 +110,7 @@ int32 FContentDirectoryMonitor::StartProcessing()
 	{
 		switch(Transaction.Action)
 		{
-			case EFileAction::Added:
+			case DirectoryWatcher::EFileAction::Added:
 				if (Settings->bAutoCreateAssets && !MountedContentPath.IsEmpty())
 				{
 					AddedFiles.Emplace(MoveTemp(Transaction));
@@ -121,12 +121,12 @@ int32 FContentDirectoryMonitor::StartProcessing()
 				}
 				break;
 
-			case EFileAction::Moved:
-			case EFileAction::Modified:
+			case DirectoryWatcher::EFileAction::Moved:
+			case DirectoryWatcher::EFileAction::Modified:
 				ModifiedFiles.Emplace(MoveTemp(Transaction));
 				break;
 
-			case EFileAction::Removed:
+			case DirectoryWatcher::EFileAction::Removed:
 				if (Settings->bAutoDeleteAssets && !MountedContentPath.IsEmpty())
 				{
 					DeletedFiles.Emplace(MoveTemp(Transaction));
@@ -159,7 +159,7 @@ UObject* AttemptImport(UClass* InFactoryType, UPackage* Package, FName InName, b
 	return Asset;
 }
 
-void FContentDirectoryMonitor::ProcessAdditions(const IAssetRegistry& Registry, const FTimeLimit& TimeLimit, TArray<UPackage*>& OutPackagesToSave, const TMap<FString, TArray<UFactory*>>& InFactoriesByExtension, FReimportFeedbackContext& Context)
+void FContentDirectoryMonitor::ProcessAdditions(const IAssetRegistry& Registry, const DirectoryWatcher::FTimeLimit& TimeLimit, TArray<UPackage*>& OutPackagesToSave, const TMap<FString, TArray<UFactory*>>& InFactoriesByExtension, FReimportFeedbackContext& Context)
 {
 	bool bCancelled = false;
 	for (int32 Index = 0; Index < AddedFiles.Num(); ++Index)
@@ -289,7 +289,7 @@ void FContentDirectoryMonitor::ProcessAdditions(const IAssetRegistry& Registry,
 	AddedFiles.Empty();
 }
 
-void FContentDirectoryMonitor::ProcessModifications(const IAssetRegistry& Registry, const FTimeLimit& TimeLimit, TArray<UPackage*>& OutPackagesToSave, FReimportFeedbackContext& Context)
+void FContentDirectoryMonitor::ProcessModifications(const IAssetRegistry& Registry, const DirectoryWatcher::FTimeLimit& TimeLimit, TArray<UPackage*>& OutPackagesToSave, FReimportFeedbackContext& Context)
 {
 	auto* ReimportManager = FReimportManager::Instance();
 
@@ -301,7 +301,7 @@ void FContentDirectoryMonitor::ProcessModifications(const IAssetRegistry& Regist
 		const FString FullFilename = Cache.GetDirectory() + Change.Filename.Get();
 
 		// Move the asset before reimporting it. We always reimport moved assets to ensure that their import path is up to date
-		if (Change.Action == EFileAction::Moved)
+		if (Change.Action == DirectoryWatcher::EFileAction::Moved)
 		{
 			const FString OldFilename = Cache.GetDirectory() + Change.MovedFromFilename.Get();
 			const auto Assets = Utils::FindAssetsPertainingToFile(Registry, OldFilename);
@@ -357,7 +357,7 @@ void FContentDirectoryMonitor::ProcessModifications(const IAssetRegistry& Regist
 				}
 			}
 		}
-		else if (Change.Action == EFileAction::Modified)
+		else if (Change.Action == DirectoryWatcher::EFileAction::Modified)
 		{
 			for (const auto& AssetData : Utils::FindAssetsPertainingToFile(Registry, FullFilename))
 			{
diff --git a/Engine/Source/Editor/UnrealEd/Private/AutoReimport/ContentDirectoryMonitor.h b/Engine/Source/Editor/UnrealEd/Private/AutoReimport/ContentDirectoryMonitor.h
--- a/Engine/Source/Editor/UnrealEd/Private/AutoReimport/ContentDirectoryMonitor.h
+++ b/Engine/Source/Editor/UnrealEd/Private/AutoReimport/ContentDirectoryMonitor.h
@@ -19,7 +19,7 @@ public:
 	 * @param InMatchRules			Rules to select what will apply to this folder
 	 * @param InMountedContentPath	(optional) Mounted content path (eg /Engine/, /Game/) to which InDirectory maps.
 	 */
-	FContentDirectoryMonitor(const FString& InDirectory, const FMatchRules& InMatchRules, const FString& InMountedContentPath = FString());
+	FContentDirectoryMonitor(const FString& InDirectory, const DirectoryWatcher::FMatchRules& InMatchRules, const FString& InMountedContentPath = FString());
 
 	/** Tick this monitor's cache to give it a chance to finish scanning for files */
 	void Tick();
@@ -28,10 +28,10 @@ public:
 	int32 StartProcessing();
 
 	/** Extract the files we need to import from our outstanding changes (happens first)*/ 
-	void ProcessAdditions(const IAssetRegistry& Registry, const FTimeLimit& TimeLimit, TArray<UPackage*>& OutPackagesToSave, const TMap<FString, TArray<UFactory*>>& InFactoriesByExtension, FReimportFeedbackContext& Context);
+	void ProcessAdditions(const IAssetRegistry& Registry, const DirectoryWatcher::FTimeLimit& TimeLimit, TArray<UPackage*>& OutPackagesToSave, const TMap<FString, TArray<UFactory*>>& InFactoriesByExtension, FReimportFeedbackContext& Context);
 
 	/** Process the outstanding changes that we have cached */
-	void ProcessModifications(const IAssetRegistry& Registry, const FTimeLimit& TimeLimit, TArray<UPackage*>& OutPackagesToSave, FReimportFeedbackContext& Context);
+	void ProcessModifications(const IAssetRegistry& Registry, const DirectoryWatcher::FTimeLimit& TimeLimit, TArray<UPackage*>& OutPackagesToSave, FReimportFeedbackContext& Context);
 
 	/** Extract the assets we need to delete from our outstanding changes (happens last) */ 
 	void ExtractAssetsToDelete(const IAssetRegistry& Registry, TArray<FAssetData>& OutAssetsToDelete);
@@ -62,13 +62,13 @@ public:
 private:
 
 	/** The file cache that monitors and reflects the content directory. */
-	FFileCache Cache;
+	DirectoryWatcher::FFileCache Cache;
 
 	/** The mounted content path for this monitor (eg /Game/) */
 	FString MountedContentPath;
 
 	/** A list of file system changes that are due to be processed */
-	TArray<FUpdateCacheTransaction> AddedFiles, ModifiedFiles, DeletedFiles;
+	TArray<DirectoryWatcher::FUpdateCacheTransaction> AddedFiles, ModifiedFiles, DeletedFiles;
 
 	/** The last time we attempted to save the cache file */
 	double LastSaveTime;
diff --git a/Engine/Source/Editor/UnrealEd/Private/AutoReimport/FileCache.cpp b/Engine/Source/Editor/UnrealEd/Private/AutoReimport/FileCache.cpp
deleted file mode 100644
diff --git a/Engine/Source/Editor/UnrealEd/Private/AutoReimport/FileCache.h b/Engine/Source/Editor/UnrealEd/Private/AutoReimport/FileCache.h
deleted file mode 100644
diff --git a/Engine/Source/Editor/UnrealEd/Private/AutoReimport/ReimportFeedbackContext.cpp b/Engine/Source/Editor/UnrealEd/Private/AutoReimport/ReimportFeedbackContext.cpp
--- a/Engine/Source/Editor/UnrealEd/Private/AutoReimport/ReimportFeedbackContext.cpp
+++ b/Engine/Source/Editor/UnrealEd/Private/AutoReimport/ReimportFeedbackContext.cpp
@@ -356,7 +356,7 @@ public:
 	void Disable()
 	{
 
-		ExpireTimeout = FTimeLimit(ExpireDuration);
+		ExpireTimeout = DirectoryWatcher::FTimeLimit(ExpireDuration);
 
 		WidgetStack->SetVisibility(EVisibility::HitTestInvisible);
 		PauseButton->SetVisibility(EVisibility::Collapsed);
@@ -367,7 +367,7 @@ public:
 	/** Enable, if previously disabled */
 	void Enable()
 	{
-		ExpireTimeout = FTimeLimit();
+		ExpireTimeout = DirectoryWatcher::FTimeLimit();
 
 		bPaused = false;
 		WidgetStack->SetVisibility(EVisibility::Visible);
@@ -436,7 +436,7 @@ private:
 private:
 
 	/** The expire timeout used to fire OnExpired. Invalid when no timeout is set. */
-	FTimeLimit ExpireTimeout;
+	DirectoryWatcher::FTimeLimit ExpireTimeout;
 
 	/** Amount of time to wait after this widget has been disabled before calling OnExpired */
 	float ExpireDuration;
