CIS fix

diff --git a/Engine/Source/Developer/DirectoryWatcher/Private/FileCache.cpp b/Engine/Source/Developer/DirectoryWatcher/Private/FileCache.cpp
--- a/Engine/Source/Developer/DirectoryWatcher/Private/FileCache.cpp
+++ b/Engine/Source/Developer/DirectoryWatcher/Private/FileCache.cpp
@@ -3,6 +3,7 @@
 #include "DirectoryWatcherPrivatePCH.h"
 #include "FileCache.h"
 #include "DirectoryWatcherModule.h"
+#include "ModuleManager.h"
 
 namespace DirectoryWatcher
 {
