Fix (hack) to get around a bug in the VS2015 compiler that prevents type info from being properly inherited by lambdas

diff --git a/Engine/Source/Developer/DirectoryWatcher/Private/DirectoryWatcherTests.cpp b/Engine/Source/Developer/DirectoryWatcher/Private/DirectoryWatcherTests.cpp
--- a/Engine/Source/Developer/DirectoryWatcher/Private/DirectoryWatcherTests.cpp
+++ b/Engine/Source/Developer/DirectoryWatcher/Private/DirectoryWatcherTests.cpp
@@ -140,6 +140,34 @@ bool FDirectoryWatcherSimpleModifyTest::RunTest(const FString& Parameters)
 	// Create a file first
 	FFileHelper::SaveStringToFile(TEXT(""), *(WorkingDir / Filename));
 
+#if _MSC_VER >= 1900
+ 	ADD_LATENT_AUTOMATION_COMMAND(FDelayedCallbackLatentCommand([=]{
+
+		// Start watching the directory
+		TSharedPtr<FDirectoryWatcherTestPayload> Test = MakeShareable(new FDirectoryWatcherTestPayload(WorkingDir));
+		void* Holder = &Test;
+
+ 		ADD_LATENT_AUTOMATION_COMMAND(FDelayedCallbackLatentCommand([=]{
+
+			// Manipulate the file
+			FFileHelper::SaveStringToFile(TEXT("Some content"), *(WorkingDir / Filename));
+
+ 			ADD_LATENT_AUTOMATION_COMMAND(FDelayedCallbackLatentCommand([=]{
+
+				TSharedPtr<FDirectoryWatcherTestPayload> LocalTest = *(const TSharedPtr<FDirectoryWatcherTestPayload>*)Holder;
+				FFileChangeData::EFileChangeAction* Action = LocalTest->ReportedChanges.Find(Filename);
+
+				if (!Action || *Action != FFileChangeData::FCA_Modified)
+				{
+					UE_LOG(LogDirectoryWatcherTests, Error, TEXT("File '%s' was not correctly reported as being modified."), Filename);
+				}
+
+			}));
+
+		}));
+
+	}));
+#else
 	ADD_LATENT_AUTOMATION_COMMAND(FDelayedCallbackLatentCommand([=]{
 
 		// Start watching the directory
@@ -164,6 +192,7 @@ bool FDirectoryWatcherSimpleModifyTest::RunTest(const FString& Parameters)
 		}));
 
 	}));
+#endif
 
 	return true;
 }
@@ -216,11 +245,13 @@ bool FDirectoryWatcherSubFolderTest::RunTest(const FString& Parameters)
 	FFileHelper::SaveStringToFile(TEXT(""), *(WorkingDir / ModifiedFilename));
 	FFileHelper::SaveStringToFile(TEXT(""), *(WorkingDir / RemovedFilename));
 
+#if _MSC_VER >= 1900
 	// Give the stream time to start up before doing the test
 	ADD_LATENT_AUTOMATION_COMMAND(FDelayedCallbackLatentCommand([=]{
 
 		// Start watching the directory
 		TSharedPtr<FDirectoryWatcherTestPayload> Test = MakeShareable(new FDirectoryWatcherTestPayload(WorkingDir));
+		void* Holder = &Test;
 
 		// Give the stream time to start up before doing the test
 		ADD_LATENT_AUTOMATION_COMMAND(FDelayedCallbackLatentCommand([=]{
@@ -234,6 +265,50 @@ bool FDirectoryWatcherSubFolderTest::RunTest(const FString& Parameters)
 
 			ADD_LATENT_AUTOMATION_COMMAND(FDelayedCallbackLatentCommand([=]{
 
+				TSharedPtr<FDirectoryWatcherTestPayload> LocalTest = *(const TSharedPtr<FDirectoryWatcherTestPayload>*)Holder;
+
+				FFileChangeData::EFileChangeAction* Action = LocalTest->ReportedChanges.Find(CreatedFilename);
+				if (!Action || *Action != FFileChangeData::FCA_Added)
+				{
+					UE_LOG(LogDirectoryWatcherTests, Error, TEXT("File '%s' was not correctly reported as being added."), CreatedFilename);
+				}
+
+				Action = LocalTest->ReportedChanges.Find(ModifiedFilename);
+				if (!Action || *Action != FFileChangeData::FCA_Modified)
+				{
+					UE_LOG(LogDirectoryWatcherTests, Error, TEXT("File '%s' was not correctly reported as being modified."), ModifiedFilename);
+				}
+
+				Action = LocalTest->ReportedChanges.Find(RemovedFilename);
+				if (!Action || *Action != FFileChangeData::FCA_Removed)
+				{
+					UE_LOG(LogDirectoryWatcherTests, Error, TEXT("File '%s' was not correctly reported as being deleted."), RemovedFilename);
+				}
+
+			}));
+
+		}));
+
+	}));
+#else
+	// Give the stream time to start up before doing the test
+	ADD_LATENT_AUTOMATION_COMMAND(FDelayedCallbackLatentCommand([=] {
+
+		// Start watching the directory
+		TSharedPtr<FDirectoryWatcherTestPayload> Test = MakeShareable(new FDirectoryWatcherTestPayload(WorkingDir));
+
+		// Give the stream time to start up before doing the test
+		ADD_LATENT_AUTOMATION_COMMAND(FDelayedCallbackLatentCommand([=] {
+
+			// Create a new file
+			FFileHelper::SaveStringToFile(TEXT(""), *(WorkingDir / CreatedFilename));
+			// Modify another file
+			FFileHelper::SaveStringToFile(TEXT("Some content"), *(WorkingDir / ModifiedFilename));
+			// Delete a file
+			IFileManager::Get().Delete(*(WorkingDir / RemovedFilename));
+
+			ADD_LATENT_AUTOMATION_COMMAND(FDelayedCallbackLatentCommand([=] {
+
 				FFileChangeData::EFileChangeAction* Action = Test->ReportedChanges.Find(CreatedFilename);
 				if (!Action || *Action != FFileChangeData::FCA_Added)
 				{
@@ -257,6 +332,7 @@ bool FDirectoryWatcherSubFolderTest::RunTest(const FString& Parameters)
 		}));
 
 	}));
+#endif
 
 	return true;
 }
@@ -269,6 +345,44 @@ bool FDirectoryWatcherNewFolderTest::RunTest(const FString& Parameters)
 	static const TCHAR* CreatedDirectory = TEXT("created");
 	static const TCHAR* RemovedDirectory = TEXT("removed");
 
+#if _MSC_VER >= 1900
+	// Give the stream time to start up before doing the test
+	ADD_LATENT_AUTOMATION_COMMAND(FDelayedCallbackLatentCommand([=] {
+
+		IFileManager::Get().MakeDirectory(*(WorkingDir / RemovedDirectory), true);
+
+		// Start watching the directory
+		TSharedPtr<FDirectoryWatcherTestPayload> Test = MakeShareable(new FDirectoryWatcherTestPayload(WorkingDir, IDirectoryWatcher::WatchOptions::IncludeDirectoryChanges));
+		void* Holder = &Test;
+
+		// Give the stream time to start up before doing the test
+		ADD_LATENT_AUTOMATION_COMMAND(FDelayedCallbackLatentCommand([=] {
+
+			IFileManager::Get().MakeDirectory(*(WorkingDir / CreatedDirectory), true);
+			IFileManager::Get().DeleteDirectory(*(WorkingDir / RemovedDirectory), true);
+
+			ADD_LATENT_AUTOMATION_COMMAND(FDelayedCallbackLatentCommand([=] {
+
+				TSharedPtr<FDirectoryWatcherTestPayload> LocalTest = *(const TSharedPtr<FDirectoryWatcherTestPayload>*)Holder;
+
+				FFileChangeData::EFileChangeAction* Action = LocalTest->ReportedChanges.Find(CreatedDirectory);
+				if (!Action || *Action != FFileChangeData::FCA_Added)
+				{
+					UE_LOG(LogDirectoryWatcherTests, Error, TEXT("Folder '%s' was not correctly reported as being added."), CreatedDirectory);
+				}
+
+				Action = LocalTest->ReportedChanges.Find(RemovedDirectory);
+				if (!Action || *Action != FFileChangeData::FCA_Removed)
+				{
+					UE_LOG(LogDirectoryWatcherTests, Error, TEXT("Folder '%s' was not correctly reported as being removed."), RemovedDirectory);
+				}
+
+			}));
+
+		}));
+
+	}));
+#else
 	// Give the stream time to start up before doing the test
 	ADD_LATENT_AUTOMATION_COMMAND(FDelayedCallbackLatentCommand([=]{
 
@@ -302,6 +416,7 @@ bool FDirectoryWatcherNewFolderTest::RunTest(const FString& Parameters)
 		}));
 
 	}));
+#endif
 
 	return true;
 }
@@ -314,11 +429,13 @@ bool FDirectoryWatcherIgnoreSubtreeTest::RunTest(const FString& Parameters)
 	static const TCHAR* ChildDirectory = TEXT("child");
 	static const TCHAR* GrandchildDirectory = TEXT("grandchild");
 
+#if _MSC_VER >= 1900
 	// Give the stream time to start up before doing the test
 	ADD_LATENT_AUTOMATION_COMMAND(FDelayedCallbackLatentCommand([=]{
 
 		// Start watching the directory
 		TSharedPtr<FDirectoryWatcherTestPayload> Test = MakeShareable(new FDirectoryWatcherTestPayload(WorkingDir, IDirectoryWatcher::WatchOptions::IgnoreChangesInSubtree | IDirectoryWatcher::WatchOptions::IncludeDirectoryChanges));
+		void* Holder = &Test;
 
 		// Give the stream time to start up before doing the test
 		ADD_LATENT_AUTOMATION_COMMAND(FDelayedCallbackLatentCommand([=]{
@@ -328,6 +445,41 @@ bool FDirectoryWatcherIgnoreSubtreeTest::RunTest(const FString& Parameters)
 
 			ADD_LATENT_AUTOMATION_COMMAND(FDelayedCallbackLatentCommand([=]{
 
+				TSharedPtr<FDirectoryWatcherTestPayload> LocalTest = *(const TSharedPtr<FDirectoryWatcherTestPayload>*)Holder;
+
+				FFileChangeData::EFileChangeAction* Action = LocalTest->ReportedChanges.Find(ChildDirectory);
+				if (!Action || *Action != FFileChangeData::FCA_Added)
+				{
+					UE_LOG(LogDirectoryWatcherTests, Error, TEXT("Folder '%s' was not correctly reported as being added."), ChildDirectory);
+				}
+
+				Action = LocalTest->ReportedChanges.Find(GrandchildDirectory);
+				if (Action)
+				{
+					UE_LOG(LogDirectoryWatcherTests, Error, TEXT("Changes to folder '%s' (creation of subfolder '%s') was reported in spite of us setting the mode 'ignore changes in subtree'."),
+						ChildDirectory, GrandchildDirectory);
+				}
+
+			}));
+
+		}));
+
+	}));
+#else
+	// Give the stream time to start up before doing the test
+	ADD_LATENT_AUTOMATION_COMMAND(FDelayedCallbackLatentCommand([=] {
+
+		// Start watching the directory
+		TSharedPtr<FDirectoryWatcherTestPayload> Test = MakeShareable(new FDirectoryWatcherTestPayload(WorkingDir, IDirectoryWatcher::WatchOptions::IgnoreChangesInSubtree | IDirectoryWatcher::WatchOptions::IncludeDirectoryChanges));
+
+		// Give the stream time to start up before doing the test
+		ADD_LATENT_AUTOMATION_COMMAND(FDelayedCallbackLatentCommand([=] {
+
+			IFileManager::Get().MakeDirectory(*(WorkingDir / ChildDirectory), true);
+			IFileManager::Get().MakeDirectory(*(WorkingDir / ChildDirectory / GrandchildDirectory), true);
+
+			ADD_LATENT_AUTOMATION_COMMAND(FDelayedCallbackLatentCommand([=] {
+
 				FFileChangeData::EFileChangeAction* Action = Test->ReportedChanges.Find(ChildDirectory);
 				if (!Action || *Action != FFileChangeData::FCA_Added)
 				{
@@ -346,6 +498,7 @@ bool FDirectoryWatcherIgnoreSubtreeTest::RunTest(const FString& Parameters)
 		}));
 
 	}));
+#endif
 
 	return true;
 }
