Move Directory Watcher tests to separate file.
Refactored MatchExtensionString method and added test case for null string handling.

diff --git a/Engine/Source/Developer/DirectoryWatcher/Private/FileCacheUtilities.cpp b/Engine/Source/Developer/DirectoryWatcher/Private/FileCacheUtilities.cpp
--- a/Engine/Source/Developer/DirectoryWatcher/Private/FileCacheUtilities.cpp
+++ b/Engine/Source/Developer/DirectoryWatcher/Private/FileCacheUtilities.cpp
@@ -3,75 +3,37 @@
 #include "DirectoryWatcherPrivatePCH.h"
 #include "FileCacheUtilities.h"
 
+namespace DirectoryWatcher
+{
+
 bool MatchExtensionString(const TCHAR* Filename, const TCHAR* Extensions)
 {
 	const int32 StrLength = FCString::Strlen(Filename);
-
 	const TCHAR* Ext = FCString::Strrchr(Filename, '.');
-	if (Ext && *(++Ext) != '\0')
-	{
-		const int32 ExtLength = StrLength - (Ext - Filename);
+	const char QueryChar = ';';
 
-		const TCHAR* Search = FCString::Strchr(Extensions, ';');
-		while(Search)
-		{
-			if (*(++Search) == '\0')
-			{
-				break;
-			}
+	if (!Ext || *(++Ext) == '\0')
+	{
+		return false;
+	}
 
-			if (FCString::Strnicmp(Search, Ext, ExtLength) == 0 && *(Search + ExtLength) == ';')
-			{
-				return true;
-			}
+	const int32 ExtLength = StrLength - (Ext - Filename);
+	const TCHAR* Search = FCString::Strchr(Extensions, QueryChar);
 
-			Search = FCString::Strchr(Search, ';');
+	while (Search && *(++Search) != '\0')
+	{
+		if (FCString::Strnicmp(Search, Ext, ExtLength) == 0
+			&& *(Search + ExtLength) == QueryChar)
+		{
+			return true;
 		}
+
+		Search = FCString::Strchr(Search, QueryChar);
 	}
 
 	return false;
 }
 
-/** Test MatchExtensionString */
-IMPLEMENT_SIMPLE_AUTOMATION_TEST(FMatchExtensionStringTest, "System.Plugins.Directory Watcher.File Cache.Extension Matching", EAutomationTestFlags::EditorContext | EAutomationTestFlags::EngineFilter)
-bool FMatchExtensionStringTest::RunTest(const FString& Parameters)
-{
-	auto Test = [](const TCHAR* Needle, const TCHAR* Haystack, bool bExpected) -> bool{
-
-		if (MatchExtensionString(Needle, Haystack) != bExpected)
-		{
-			if (bExpected)
-			{
-				UE_LOG(LogFileCache, Error, TEXT("Unable to match '%s' in '%s'"), Needle, Haystack);
-			}
-			else
-			{
-				UE_LOG(LogFileCache, Error, TEXT("Erroneously matched '%s' in '%s'"), Needle, Haystack);
-			}
-			return false;
-		}
-		return true;
-	};
-
-	bool Result = true;
-	Result = Test(TEXT("bla.txt"), TEXT(";;"), false) && Result;
-	Result = Test(TEXT("bla.txt"), TEXT(";"), false) && Result;
-	Result = Test(TEXT("bla.txt"), TEXT(""), false) && Result;
-	Result = Test(TEXT("bla.txt"), TEXT(";txt;"), true) && Result;
-	Result = Test(TEXT("bla.text"), TEXT(";txt;"), false) && Result;
-	Result = Test(TEXT("bla.txt1"), TEXT(";txt;"), false) && Result;
-	Result = Test(TEXT("bla."), TEXT(";bla;"), false) && Result;
-	Result = Test(TEXT("bla.png"), TEXT(";png;txt;"), true) && Result;
-	Result = Test(TEXT("bla.txt"), TEXT(";png;txt;"), true) && Result;
-	Result = Test(TEXT("/folder.bin/bla.txt"), TEXT(";png;txt;"), true) && Result;
-	Result = Test(TEXT("/folder.bin/bla"), TEXT(";png;bin;"), false) && Result;
-
-	return true;
-}
-
-namespace DirectoryWatcher
-{
-
 struct FWildcardRule : IMatchRule
 {
 	FWildcardString WildcardString;
diff --git a/Engine/Source/Developer/DirectoryWatcher/Private/Tests/DirectoryWatcherTests.cpp b/Engine/Source/Developer/DirectoryWatcher/Private/Tests/DirectoryWatcherTests.cpp
new file mode 100644
--- /dev/null
+++ b/Engine/Source/Developer/DirectoryWatcher/Private/Tests/DirectoryWatcherTests.cpp
@@ -0,0 +1,40 @@
+// Copyright 1998-2015 Epic Games, Inc. All Rights Reserved.
+
+#include "DirectoryWatcherPrivatePCH.h"
+#include "FileCacheUtilities.h"
+
+/** Test MatchExtensionString */
+IMPLEMENT_SIMPLE_AUTOMATION_TEST(FMatchExtensionStringTest, "System.Plugins.Directory Watcher.File Cache.Extension Matching", EAutomationTestFlags::EditorContext | EAutomationTestFlags::EngineFilter)
+
+namespace MatchExtensionTestHelpers
+{
+	FString ExpectedMessage = TEXT("Matched");
+	FString NotExpectedMessage = TEXT("Did not erroneously match");
+
+	bool Test(FMatchExtensionStringTest* Test, const TCHAR* Needle, const TCHAR* Haystack, bool bExpected)
+	{
+		bool bResult = DirectoryWatcher::MatchExtensionString(Needle, Haystack);
+		Test->TestEqual(FString::Printf(TEXT("%s '%s' in '%s'"), bExpected ? *ExpectedMessage : *NotExpectedMessage, Needle, Haystack), bExpected, bResult);
+
+		return bResult == bExpected;
+	}
+}
+
+bool FMatchExtensionStringTest::RunTest(const FString& Parameters)
+{
+	bool Result = true;
+	Result = MatchExtensionTestHelpers::Test(this, TEXT("blatxt"), TEXT(";txt;"), false) && Result;  //FCString::Strrchr in MatchExtentionString returns Null if Filename has no periods
+	Result = MatchExtensionTestHelpers::Test(this, TEXT("bla.txt"), TEXT(";;"), false) && Result;
+	Result = MatchExtensionTestHelpers::Test(this, TEXT("bla.txt"), TEXT(";"), false) && Result;
+	Result = MatchExtensionTestHelpers::Test(this, TEXT("bla.txt"), TEXT(""), false) && Result;
+	Result = MatchExtensionTestHelpers::Test(this, TEXT("bla.txt"), TEXT(";txt;"), true) && Result;
+	Result = MatchExtensionTestHelpers::Test(this, TEXT("bla.text"), TEXT(";txt;"), false) && Result;
+	Result = MatchExtensionTestHelpers::Test(this, TEXT("bla.txt1"), TEXT(";txt;"), false) && Result;
+	Result = MatchExtensionTestHelpers::Test(this, TEXT("bla."), TEXT(";bla;"), false) && Result;
+	Result = MatchExtensionTestHelpers::Test(this, TEXT("bla.png"), TEXT(";png;txt;"), true) && Result;
+	Result = MatchExtensionTestHelpers::Test(this, TEXT("bla.txt"), TEXT(";png;txt;"), true) && Result;
+	Result = MatchExtensionTestHelpers::Test(this, TEXT("/folder.bin/bla.txt"), TEXT(";png;txt;"), true) && Result;
+	Result = MatchExtensionTestHelpers::Test(this, TEXT("/folder.bin/bla"), TEXT(";png;bin;"), false) && Result;
+
+	return Result;
+}
\ No newline at end of file
diff --git a/Engine/Source/Developer/DirectoryWatcher/Public/FileCacheUtilities.h b/Engine/Source/Developer/DirectoryWatcher/Public/FileCacheUtilities.h
--- a/Engine/Source/Developer/DirectoryWatcher/Public/FileCacheUtilities.h
+++ b/Engine/Source/Developer/DirectoryWatcher/Public/FileCacheUtilities.h
@@ -7,6 +7,9 @@ DEFINE_LOG_CATEGORY_STATIC(LogFileCache, Log, All);
 namespace DirectoryWatcher
 {
 
+/** Helper method to determine whether filename matches expected extensions */
+bool MatchExtensionString(const TCHAR* Filename, const TCHAR* Extensions);
+
 /** An immutable string with a cached CRC for efficient comparison with other strings */
 struct FImmutableString
 {
