Fixed missing param to log call (merge error)

diff --git a/Engine/Source/Developer/DirectoryWatcher/Private/Windows/DirectoryWatcherWindows.cpp b/Engine/Source/Developer/DirectoryWatcher/Private/Windows/DirectoryWatcherWindows.cpp
--- a/Engine/Source/Developer/DirectoryWatcher/Private/Windows/DirectoryWatcherWindows.cpp
+++ b/Engine/Source/Developer/DirectoryWatcher/Private/Windows/DirectoryWatcherWindows.cpp
@@ -62,7 +62,7 @@ bool FDirectoryWatcherWindows::RegisterDirectoryChangedCallback_Handle( const FS
 			uint32 Error = GetLastError();
 			TCHAR ErrorMsg[1024];
 			FPlatformMisc::GetSystemErrorMessage(ErrorMsg, 1024, Error);
-			UE_LOG(LogDirectoryWatcher, Warning, TEXT("Failed to begin reading directory changes for %s. Error: %s (0x%08x)"), ErrorMsg, Error);
+			UE_LOG(LogDirectoryWatcher, Warning, TEXT("Failed to begin reading directory changes for %s. Error: %s (0x%08x)"), *Directory, ErrorMsg, Error);
 
 			delete Request;
 			NumRequests--;
