Downgraded some warnings that caused commandlet spam

diff --git a/Engine/Source/Developer/DirectoryWatcher/Private/Windows/DirectoryWatchRequestWindows.cpp b/Engine/Source/Developer/DirectoryWatcher/Private/Windows/DirectoryWatchRequestWindows.cpp
--- a/Engine/Source/Developer/DirectoryWatcher/Private/Windows/DirectoryWatchRequestWindows.cpp
+++ b/Engine/Source/Developer/DirectoryWatcher/Private/Windows/DirectoryWatchRequestWindows.cpp
@@ -175,7 +175,7 @@ void FDirectoryWatchRequestWindows::ProcessChange(uint32 Error, uint32 NumBytes)
 		// The operation was aborted, likely due to EndWatchRequest canceling it.
 		// Mark the request for delete so it can be cleaned up next tick.
 		bPendingDelete = true;
-		UE_LOG(LogDirectoryWatcher, Warning, TEXT("A directory notification for '%s' was aborted."), *Directory);
+		UE_LOG(LogDirectoryWatcher, Log, TEXT("A directory notification for '%s' was aborted."), *Directory);
 		return; 
 	}
 
@@ -190,7 +190,7 @@ void FDirectoryWatchRequestWindows::ProcessChange(uint32 Error, uint32 NumBytes)
 
 	if ( !bValidNotification )
 	{
-		UE_LOG(LogDirectoryWatcher, Warning, TEXT("A directory notification failed for '%s' because it was empty or there was a buffer overflow. Attemping another request..."), *Directory);
+		UE_LOG(LogDirectoryWatcher, Log, TEXT("A directory notification failed for '%s' because it was empty or there was a buffer overflow. Attemping another request..."), *Directory);
 	}
 
 	// Start up another read
@@ -211,7 +211,7 @@ void FDirectoryWatchRequestWindows::ProcessChange(uint32 Error, uint32 NumBytes)
 		::CloseHandle(DirectoryHandle);
 		DirectoryHandle = INVALID_HANDLE_VALUE;
 		bPendingDelete = true;
-		UE_LOG(LogDirectoryWatcher, Warning, TEXT("A directory notification failed for '%s', and we were unable to create a new request."), *Directory);
+		UE_LOG(LogDirectoryWatcher, Log, TEXT("A directory notification failed for '%s', and we were unable to create a new request."), *Directory);
 		return;
 	}
 
