Linux: don't let inotify failure break cooks.

- Errors cause Cook commandlet to exit with 1, while this error is arguably benign.

diff --git a/Engine/Source/Developer/DirectoryWatcher/Private/Linux/DirectoryWatchRequestLinux.cpp b/Engine/Source/Developer/DirectoryWatcher/Private/Linux/DirectoryWatchRequestLinux.cpp
--- a/Engine/Source/Developer/DirectoryWatcher/Private/Linux/DirectoryWatchRequestLinux.cpp
+++ b/Engine/Source/Developer/DirectoryWatcher/Private/Linux/DirectoryWatchRequestLinux.cpp
@@ -51,8 +51,17 @@ bool FDirectoryWatchRequestLinux::Init(const FString& InDirectory, uint32 Flags)
 
 	if (FileDescriptor == -1)
 	{
-		// Failed to init inotify
-		UE_LOG(LogDirectoryWatcher, Error, TEXT("Failed to init inotify"));
+		// Failed to init inotify. The number of inotify instances on the default Linux is vastly
+		// insufficient, so do not make this an error which can break cooking
+		int ErrNo = errno;
+		if (ErrNo == EMFILE)
+		{
+			UE_LOG(LogDirectoryWatcher, Warning, TEXT("Failed to init inotify (ran out of inotify instances, consider increasing system limits)"));
+		}
+		else
+		{
+			UE_LOG(LogDirectoryWatcher, Error, TEXT("Failed to init inotify (errno=%d, %s)"), ErrNo, UTF8_TO_TCHAR(strerror(ErrNo)));
+		}
 		return false;
 	}
 
