Modified Mac DirectoryWatcher to use the same method of converting CFStringRef to FString that ApplePlatformFile uses to make sure unicode representation of paths is identical.

diff --git a/Engine/Source/Developer/DirectoryWatcher/Private/Mac/DirectoryWatchRequestMac.cpp b/Engine/Source/Developer/DirectoryWatcher/Private/Mac/DirectoryWatchRequestMac.cpp
--- a/Engine/Source/Developer/DirectoryWatcher/Private/Mac/DirectoryWatchRequestMac.cpp
+++ b/Engine/Source/Developer/DirectoryWatcher/Private/Mac/DirectoryWatchRequestMac.cpp
@@ -144,8 +144,6 @@ void FDirectoryWatchRequestMac::ProcessChanges( size_t EventCount, void* EventPa
 		return;
 	}
 
-	TCHAR FilePath[MAC_MAX_PATH];
-
 	CFArrayRef EventPathArray = (CFArrayRef)EventPaths;
 
 	for( size_t EventIndex = 0; EventIndex < EventCount; ++EventIndex )
@@ -205,13 +203,12 @@ void FDirectoryWatchRequestMac::ProcessChanges( size_t EventCount, void* EventPa
 			continue;
 		}
 
-		FPlatformString::CFStringToTCHAR((CFStringRef)CFArrayGetValueAtIndex(EventPathArray,EventIndex), FilePath);
-
-		if( bFileNeedsChecking && !FPlatformFileManager::Get().GetPlatformFile().FileExists(FilePath) )
+		const FString FilePath = UTF8_TO_TCHAR(([[(NSString*)(CFStringRef)CFArrayGetValueAtIndex(EventPathArray,EventIndex) precomposedStringWithCanonicalMapping] cStringUsingEncoding:NSUTF8StringEncoding]));
+		if( bFileNeedsChecking && !FPlatformFileManager::Get().GetPlatformFile().FileExists(*FilePath) )
 		{
 			Action = FFileChangeData::FCA_Removed;
 		}
 
-		new(FileChanges) FFileChangeData(FString(FilePath), Action);
+		new(FileChanges) FFileChangeData(FilePath, Action);
 	}
 }
