Watch directories in AssetRegistry and allow watches on the same directory with different flags
Improve check to determine correct folder hierarchy

diff --git a/Engine/Source/Developer/DirectoryWatcher/Private/Linux/DirectoryWatchRequestLinux.cpp b/Engine/Source/Developer/DirectoryWatcher/Private/Linux/DirectoryWatchRequestLinux.cpp
--- a/Engine/Source/Developer/DirectoryWatcher/Private/Linux/DirectoryWatchRequestLinux.cpp
+++ b/Engine/Source/Developer/DirectoryWatcher/Private/Linux/DirectoryWatchRequestLinux.cpp
@@ -35,7 +35,6 @@ bool FDirectoryWatchRequestLinux::Init(const FString& InDirectory, uint32 Flags)
 	}
 
 	Directory = InDirectory;
-	bIncludeDirectoryChanges = (Flags & IDirectoryWatcher::WatchOptions::IncludeDirectoryChanges) != 0;
 	bWatchSubtree = (Flags & IDirectoryWatcher::WatchOptions::IgnoreChangesInSubtree) == 0;
 
 	if (bRunning)
@@ -66,16 +65,16 @@ bool FDirectoryWatchRequestLinux::Init(const FString& InDirectory, uint32 Flags)
 	return true;
 }
 
-FDelegateHandle FDirectoryWatchRequestLinux::AddDelegate(const IDirectoryWatcher::FDirectoryChanged& InDelegate)
+FDelegateHandle FDirectoryWatchRequestLinux::AddDelegate(const IDirectoryWatcher::FDirectoryChanged& InDelegate, uint32 Flags)
 {
-	Delegates.Add(InDelegate);
-	return Delegates.Last().GetHandle();
+	Delegates.Emplace(InDelegate, Flags);
+	return Delegates.Last().Key.GetHandle();
 }
 
 bool FDirectoryWatchRequestLinux::RemoveDelegate(FDelegateHandle InHandle)
 {
-	return Delegates.RemoveAll([=](const IDirectoryWatcher::FDirectoryChanged& Delegate) {
-		return Delegate.GetHandle() == InHandle;
+	return Delegates.RemoveAll([=](const FWatchDelegate& Delegate) {
+		return Delegate.Key.GetHandle() == InHandle;
 	}) != 0;
 }
 
@@ -97,9 +96,28 @@ void FDirectoryWatchRequestLinux::ProcessPendingNotifications()
 	// Trigger all listening delegates with the files that have changed
 	if (FileChanges.Num() > 0)
 	{
-		for (int32 DelegateIdx = 0; DelegateIdx < Delegates.Num(); ++DelegateIdx)
+		TMap<uint32, TArray<FFileChangeData>> FileChangeCache;
+		for (const FWatchDelegate& Delegate : Delegates)
 		{
-			Delegates[DelegateIdx].Execute(FileChanges);
+			// Filter list of all file changes down to ones that just match this delegate's flags
+			TArray<FFileChangeData>* CachedChanges = FileChangeCache.Find(Delegate.Value);
+			if (CachedChanges)
+			{
+				Delegate.Key.Execute(*CachedChanges);
+			}
+			else
+			{
+				const bool bIncludeDirs = (Delegate.Value & IDirectoryWatcher::WatchOptions::IncludeDirectoryChanges) != 0;
+				TArray<FFileChangeData>& Changes = FileChangeCache.Add(Delegate.Value);
+				for (const TPair<FFileChangeData, bool>& FileChangeData : FileChanges)
+				{
+					if (!FileChangeData.Value || bIncludeDirs)
+					{
+						Changes.Add(FileChangeData.Key);
+					}
+				}
+				Delegate.Key.Execute(Changes);
+			}
 		}
 
 		FileChanges.Empty();
@@ -216,9 +234,6 @@ void FDirectoryWatchRequestLinux::ProcessChanges()
 		{
 			Event = reinterpret_cast<const struct inotify_event *>(Ptr);
 
-			// some events can be ignored to match other implementations and tests
-			bool bIgnoreEvent = false;
-
 			// skip if overflowed
 			if (Event->wd != -1 && (Event->mask & IN_Q_OVERFLOW) == 0)
 			{
@@ -240,8 +255,6 @@ void FDirectoryWatchRequestLinux::ProcessChanges()
 						if (Event->mask & IN_ISDIR)
 						{
 							WatchDirectoryTree(AffectedFile);
-							// to be in sync with Windows implementation, ignore events about creating directories unless told so
-							bIgnoreEvent = !bIncludeDirectoryChanges;
 						}
 
 						Action = FFileChangeData::FCA_Added;
@@ -259,8 +272,6 @@ void FDirectoryWatchRequestLinux::ProcessChanges()
 						if (Event->mask & IN_ISDIR)
 						{
 							UnwatchDirectoryTree(EventPath);
-							// to be in sync with Windows implementation, ignore events about deleting directories unless told so
-							bIgnoreEvent = !bIncludeDirectoryChanges;
 						}
 						else
 						{
@@ -279,16 +290,14 @@ void FDirectoryWatchRequestLinux::ProcessChanges()
 						if (Event->mask & IN_ISDIR)
 						{
 							UnwatchDirectoryTree(AffectedFile);
-							// to be in sync with Windows implementation, ignore events about deleting directories unless told so
-							bIgnoreEvent = !bIncludeDirectoryChanges;
 						}
 
 						Action = FFileChangeData::FCA_Removed;
 					}
 
-					if (!bIgnoreEvent && Event->len)
+					if (Event->len)
 					{
-						new (FileChanges) FFileChangeData(AffectedFile, Action);
+						FileChanges.Emplace(FFileChangeData(AffectedFile, Action), (Event->mask & IN_ISDIR) != 0);
 					}
 				}
 			}
diff --git a/Engine/Source/Developer/DirectoryWatcher/Private/Linux/DirectoryWatchRequestLinux.h b/Engine/Source/Developer/DirectoryWatcher/Private/Linux/DirectoryWatchRequestLinux.h
--- a/Engine/Source/Developer/DirectoryWatcher/Private/Linux/DirectoryWatchRequestLinux.h
+++ b/Engine/Source/Developer/DirectoryWatcher/Private/Linux/DirectoryWatchRequestLinux.h
@@ -41,7 +41,7 @@ public:
 	bool Init(const FString& InDirectory, uint32 Flags);
 
 	/** Adds a delegate to get fired when the directory changes */
-	FDelegateHandle AddDelegate( const IDirectoryWatcher::FDirectoryChanged& InDelegate );
+	FDelegateHandle AddDelegate( const IDirectoryWatcher::FDirectoryChanged& InDelegate, uint32 Flags );
 	/** Removes a delegate to get fired when the directory changes */
 	bool RemoveDelegate( FDelegateHandle InHandle );
 	/** Returns true if this request has any delegates listening to directory changes */
@@ -64,9 +64,6 @@ private:
 	bool bRunning;
 	bool bEndWatchRequestInvoked;
 
-	/** Whether to report directory creation/deletion changes. */
-	bool bIncludeDirectoryChanges;
-
 	/** Whether or not watch subtree. */
 	bool bWatchSubtree;
 
@@ -80,8 +77,11 @@ private:
 
 	int NotifyFilter;
 
-	TArray<IDirectoryWatcher::FDirectoryChanged> Delegates;
-	TArray<FFileChangeData> FileChanges;
+	/** A delegate with its corresponding IDirectoryWatcher::WatchOptions flags */
+	typedef TPair<IDirectoryWatcher::FDirectoryChanged, uint32> FWatchDelegate;
+	TArray<FWatchDelegate> Delegates;
+	/** Each FFileChangeData tracks whether it is a directory or not */
+	TArray<TPair<FFileChangeData, bool>> FileChanges;
 
 	void Shutdown();
 	void ProcessChanges();
diff --git a/Engine/Source/Developer/DirectoryWatcher/Private/Linux/DirectoryWatcherLinux.cpp b/Engine/Source/Developer/DirectoryWatcher/Private/Linux/DirectoryWatcherLinux.cpp
--- a/Engine/Source/Developer/DirectoryWatcher/Private/Linux/DirectoryWatcherLinux.cpp
+++ b/Engine/Source/Developer/DirectoryWatcher/Private/Linux/DirectoryWatcherLinux.cpp
@@ -67,7 +67,7 @@ bool FDirectoryWatcherLinux::RegisterDirectoryChangedCallback_Handle(const FStri
 		RequestMap.Add(Directory, Request);
 	}
 
-	OutHandle = Request->AddDelegate(InDelegate);
+	OutHandle = Request->AddDelegate(InDelegate, Flags);
 
 	return true;
 }
diff --git a/Engine/Source/Developer/DirectoryWatcher/Private/Mac/DirectoryWatchRequestMac.cpp b/Engine/Source/Developer/DirectoryWatcher/Private/Mac/DirectoryWatchRequestMac.cpp
--- a/Engine/Source/Developer/DirectoryWatcher/Private/Mac/DirectoryWatchRequestMac.cpp
+++ b/Engine/Source/Developer/DirectoryWatcher/Private/Mac/DirectoryWatchRequestMac.cpp
@@ -15,11 +15,9 @@ void DirectoryWatchMacCallback( ConstFSEventStreamRef StreamRef, void* WatchRequ
 
 // ============================================================================================================================
 
-FDirectoryWatchRequestMac::FDirectoryWatchRequestMac(uint32 Flags)
+FDirectoryWatchRequestMac::FDirectoryWatchRequestMac()
 :	bRunning(false)
 ,	bEndWatchRequestInvoked(false)
-,	bIncludeDirectoryEvents((Flags & IDirectoryWatcher::WatchOptions::IncludeDirectoryChanges) != 0)
-,	bIgnoreChangesInSubtree((Flags & IDirectoryWatcher::WatchOptions::IgnoreChangesInSubtree) != 0)
 {
 }
 
@@ -99,16 +97,16 @@ bool FDirectoryWatchRequestMac::Init(const FString& InDirectory)
 	return true;
 }
 
-FDelegateHandle FDirectoryWatchRequestMac::AddDelegate( const IDirectoryWatcher::FDirectoryChanged& InDelegate )
+FDelegateHandle FDirectoryWatchRequestMac::AddDelegate( const IDirectoryWatcher::FDirectoryChanged& InDelegate, uint32 Flags )
 {
-	Delegates.Add(InDelegate);
-	return Delegates.Last().GetHandle();
+	Delegates.Emplace(InDelegate, Flags);
+	return Delegates.Last().Key.GetHandle();
 }
 
 bool FDirectoryWatchRequestMac::RemoveDelegate( FDelegateHandle InHandle )
 {
-	return Delegates.RemoveAll([=](const IDirectoryWatcher::FDirectoryChanged& Delegate) {
-		return Delegate.GetHandle() == InHandle;
+	return Delegates.RemoveAll([=](const FWatchDelegate& Delegate) {
+		return Delegate.Key.GetHandle() == InHandle;
 	}) != 0;
 }
 
@@ -127,9 +125,29 @@ void FDirectoryWatchRequestMac::ProcessPendingNotifications()
 	// Trigger all listening delegates with the files that have changed
 	if ( FileChanges.Num() > 0 )
 	{
-		for (int32 DelegateIdx = 0; DelegateIdx < Delegates.Num(); ++DelegateIdx)
+		TMap<uint32, TArray<FFileChangeData>> FileChangeCache;
+		for (const FWatchDelegate& Delegate : Delegates)
 		{
-			Delegates[DelegateIdx].Execute(FileChanges);
+			// Filter list of all file changes down to ones that just match this delegate's flags
+			TArray<FFileChangeData>* CachedChanges = FileChangeCache.Find(Delegate.Value);
+			if (CachedChanges)
+			{
+				Delegate.Key.Execute(*CachedChanges);
+			}
+			else
+			{
+				const bool bIncludeDirs = (Delegate.Value & IDirectoryWatcher::WatchOptions::IncludeDirectoryChanges) != 0;
+				TArray<FFileChangeData>& Changes = FileChangeCache.Add(Delegate.Value);
+				for (const TPair<FFileChangeData, bool>& FileChangeData : FileChanges)
+				{
+					// @todo support IgnoreChangesInSubtree
+					if (!FileChangeData.Value || bIncludeDirs)
+					{
+						Changes.Add(FileChangeData.Key);
+					}
+				}
+				Delegate.Key.Execute(Changes);
+			}
 		}
 
 		FileChanges.Empty();
@@ -149,13 +167,10 @@ void FDirectoryWatchRequestMac::ProcessChanges( size_t EventCount, void* EventPa
 	for( size_t EventIndex = 0; EventIndex < EventCount; ++EventIndex )
 	{
 		const FSEventStreamEventFlags Flags = EventFlags[EventIndex];
-		if( !(Flags & kFSEventStreamEventFlagItemIsFile) )
+		if( !(Flags & kFSEventStreamEventFlagItemIsFile) && !(Flags & kFSEventStreamEventFlagItemIsDir) )
 		{
-			if( !bIncludeDirectoryEvents || !(Flags & kFSEventStreamEventFlagItemIsDir))
-			{
-				// events about directories and symlinks don't concern us
-				continue;
-			}
+			// events about symlinks don't concern us
+			continue;
 		}
 
 		// Warning: some events have more than one of created, removed nd modified flag.
@@ -209,6 +224,6 @@ void FDirectoryWatchRequestMac::ProcessChanges( size_t EventCount, void* EventPa
 			Action = FFileChangeData::FCA_Removed;
 		}
 
-		new(FileChanges) FFileChangeData(FilePath, Action);
+		FileChanges.Emplace(FFileChangeData(FilePath, Action), (Flags & kFSEventStreamEventFlagItemIsDir) != 0);
 	}
 }
diff --git a/Engine/Source/Developer/DirectoryWatcher/Private/Mac/DirectoryWatchRequestMac.h b/Engine/Source/Developer/DirectoryWatcher/Private/Mac/DirectoryWatchRequestMac.h
--- a/Engine/Source/Developer/DirectoryWatcher/Private/Mac/DirectoryWatchRequestMac.h
+++ b/Engine/Source/Developer/DirectoryWatcher/Private/Mac/DirectoryWatchRequestMac.h
@@ -11,14 +11,14 @@
 class FDirectoryWatchRequestMac
 {
 public:
-	FDirectoryWatchRequestMac(uint32 Flags);
+	FDirectoryWatchRequestMac();
 	virtual ~FDirectoryWatchRequestMac();
 
 	/** Sets up the directory handle and request information */
 	bool Init(const FString& InDirectory);
 
 	/** Adds a delegate to get fired when the directory changes */
-	FDelegateHandle AddDelegate( const IDirectoryWatcher::FDirectoryChanged& InDelegate );
+	FDelegateHandle AddDelegate( const IDirectoryWatcher::FDirectoryChanged& InDelegate, uint32 Flags );
 	/** Removes a delegate to get fired when the directory changes */
 	bool RemoveDelegate( FDelegateHandle InHandle );
 	/** Returns true if this request has any delegates listening to directory changes */
@@ -33,11 +33,12 @@ private:
 	FSEventStreamRef	EventStream;
 	bool				bRunning;
 	bool				bEndWatchRequestInvoked;
-	bool				bIncludeDirectoryEvents;
-	bool				bIgnoreChangesInSubtree;
 
-	TArray<IDirectoryWatcher::FDirectoryChanged> Delegates;
-	TArray<FFileChangeData> FileChanges;
+	/** A delegate with its corresponding IDirectoryWatcher::WatchOptions flags */
+	typedef TPair<IDirectoryWatcher::FDirectoryChanged, uint32> FWatchDelegate;
+	TArray<FWatchDelegate> Delegates;
+	/** Each FFileChangeData tracks whether it is a directory or not */
+	TArray<TPair<FFileChangeData, bool>> FileChanges;
 
 	friend void DirectoryWatchMacCallback( ConstFSEventStreamRef StreamRef, void* WatchRequestPtr, size_t EventCount, void* EventPaths, const FSEventStreamEventFlags EventFlags[], const FSEventStreamEventId EventIDs[] );
 
diff --git a/Engine/Source/Developer/DirectoryWatcher/Private/Mac/DirectoryWatcherMac.cpp b/Engine/Source/Developer/DirectoryWatcher/Private/Mac/DirectoryWatcherMac.cpp
--- a/Engine/Source/Developer/DirectoryWatcher/Private/Mac/DirectoryWatcherMac.cpp
+++ b/Engine/Source/Developer/DirectoryWatcher/Private/Mac/DirectoryWatcherMac.cpp
@@ -53,7 +53,7 @@ bool FDirectoryWatcherMac::RegisterDirectoryChangedCallback_Handle( const FStrin
 	}
 	else
 	{
-		Request = new FDirectoryWatchRequestMac(Flags);
+		Request = new FDirectoryWatchRequestMac();
 		NumRequests++;
 
 		// Begin reading directory changes
@@ -68,7 +68,7 @@ bool FDirectoryWatcherMac::RegisterDirectoryChangedCallback_Handle( const FStrin
 		RequestMap.Add(Directory, Request);
 	}
 
-	OutHandle = Request->AddDelegate(InDelegate);
+	OutHandle = Request->AddDelegate(InDelegate, Flags);
 
 	return true;
 }
diff --git a/Engine/Source/Developer/DirectoryWatcher/Private/Windows/DirectoryWatcherWindows.cpp b/Engine/Source/Developer/DirectoryWatcher/Private/Windows/DirectoryWatcherWindows.cpp
--- a/Engine/Source/Developer/DirectoryWatcher/Private/Windows/DirectoryWatcherWindows.cpp
+++ b/Engine/Source/Developer/DirectoryWatcher/Private/Windows/DirectoryWatcherWindows.cpp
@@ -13,7 +13,7 @@ FDirectoryWatcherWindows::~FDirectoryWatcherWindows()
 	if ( RequestMap.Num() != 0 )
 	{
 		// Delete any remaining requests here. These requests are likely from modules which are still loaded at the time that this module unloads.
-		for (TMap<FString, FDirectoryWatchRequestWindows*>::TConstIterator RequestIt(RequestMap); RequestIt; ++RequestIt)
+		for (TMap<FDirectoryWithFlags, FDirectoryWatchRequestWindows*>::TConstIterator RequestIt(RequestMap); RequestIt; ++RequestIt)
 		{
 			if ( ensure(RequestIt.Value()) )
 			{
@@ -42,7 +42,8 @@ FDirectoryWatcherWindows::~FDirectoryWatcherWindows()
 
 bool FDirectoryWatcherWindows::RegisterDirectoryChangedCallback_Handle( const FString& Directory, const FDirectoryChanged& InDelegate, FDelegateHandle& Handle, uint32 Flags )
 {
-	FDirectoryWatchRequestWindows** RequestPtr = RequestMap.Find(Directory);
+	const FDirectoryWithFlags DirectoryKey(Directory, Flags);
+	FDirectoryWatchRequestWindows** RequestPtr = RequestMap.Find(DirectoryKey);
 	FDirectoryWatchRequestWindows* Request = NULL;
 	
 	if ( RequestPtr )
@@ -70,7 +71,7 @@ bool FDirectoryWatcherWindows::RegisterDirectoryChangedCallback_Handle( const FS
 			return false;
 		}
 
-		RequestMap.Add(Directory, Request);
+		RequestMap.Add(DirectoryKey, Request);
 	}
 
 	Handle = Request->AddDelegate(InDelegate);
@@ -80,30 +81,28 @@ bool FDirectoryWatcherWindows::RegisterDirectoryChangedCallback_Handle( const FS
 
 bool FDirectoryWatcherWindows::UnregisterDirectoryChangedCallback_Handle( const FString& Directory, FDelegateHandle InHandle )
 {
-	FDirectoryWatchRequestWindows** RequestPtr = RequestMap.Find(Directory);
-	
-	if ( RequestPtr )
+	for (const auto& RequestPair : RequestMap)
 	{
-		// There should be no NULL entries in the map
-		check (*RequestPtr);
-
-		FDirectoryWatchRequestWindows* Request = *RequestPtr;
-
-		if ( Request->RemoveDelegate(InHandle) )
+		if (RequestPair.Key.Key == Directory)
 		{
-			if ( !Request->HasDelegates() )
+			// There should be no NULL entries in the map
+			check(RequestPair.Value);
+
+			if (RequestPair.Value->RemoveDelegate(InHandle))
 			{
-				// Remove from the active map and add to the pending delete list
-				RequestMap.Remove(Directory);
-				RequestsPendingDelete.AddUnique(Request);
+				if (!RequestPair.Value->HasDelegates())
+				{
+					// Remove from the active map and add to the pending delete list
+					RequestMap.Remove(RequestPair.Key);
+					RequestsPendingDelete.AddUnique(RequestPair.Value);
 
-				// Signal to end the watch which will mark this request for deletion
-				Request->EndWatchRequest();
-			}
+					// Signal to end the watch which will mark this request for deletion
+					RequestPair.Value->EndWatchRequest();
+				}
 
-			return true;
+				return true;
+			}
 		}
-		
 	}
 
 	return false;
@@ -112,10 +111,10 @@ bool FDirectoryWatcherWindows::UnregisterDirectoryChangedCallback_Handle( const
 void FDirectoryWatcherWindows::Tick( float DeltaSeconds )
 {
 	TArray<HANDLE> DirectoryHandles;
-	TMap<FString, FDirectoryWatchRequestWindows*> InvalidRequestsToDelete;
+	TMap<FDirectoryWithFlags, FDirectoryWatchRequestWindows*> InvalidRequestsToDelete;
 
 	// Find all handles to listen to and invalid requests to delete
-	for (TMap<FString, FDirectoryWatchRequestWindows*>::TConstIterator RequestIt(RequestMap); RequestIt; ++RequestIt)
+	for (TMap<FDirectoryWithFlags, FDirectoryWatchRequestWindows*>::TConstIterator RequestIt(RequestMap); RequestIt; ++RequestIt)
 	{
 		if ( RequestIt.Value()->IsPendingDelete() )
 		{
@@ -128,7 +127,7 @@ void FDirectoryWatcherWindows::Tick( float DeltaSeconds )
 	}
 
 	// Remove all invalid requests from the request map and add them to the pending delete list so they will be deleted below
-	for (TMap<FString, FDirectoryWatchRequestWindows*>::TConstIterator RequestIt(InvalidRequestsToDelete); RequestIt; ++RequestIt)
+	for (TMap<FDirectoryWithFlags, FDirectoryWatchRequestWindows*>::TConstIterator RequestIt(InvalidRequestsToDelete); RequestIt; ++RequestIt)
 	{
 		RequestMap.Remove(RequestIt.Key());
 		RequestsPendingDelete.AddUnique(RequestIt.Value());
@@ -155,7 +154,7 @@ void FDirectoryWatcherWindows::Tick( float DeltaSeconds )
 	}
 
 	// Finally, trigger any file change notification delegates
-	for (TMap<FString, FDirectoryWatchRequestWindows*>::TConstIterator RequestIt(RequestMap); RequestIt; ++RequestIt)
+	for (TMap<FDirectoryWithFlags, FDirectoryWatchRequestWindows*>::TConstIterator RequestIt(RequestMap); RequestIt; ++RequestIt)
 	{
 		RequestIt.Value()->ProcessPendingNotifications();
 	}
diff --git a/Engine/Source/Developer/DirectoryWatcher/Private/Windows/DirectoryWatcherWindows.h b/Engine/Source/Developer/DirectoryWatcher/Private/Windows/DirectoryWatcherWindows.h
--- a/Engine/Source/Developer/DirectoryWatcher/Private/Windows/DirectoryWatcherWindows.h
+++ b/Engine/Source/Developer/DirectoryWatcher/Private/Windows/DirectoryWatcherWindows.h
@@ -12,6 +12,9 @@ class FDirectoryWatchRequestWindows;
 
 class FDirectoryWatcherWindows : public IDirectoryWatcher
 {
+private:
+	/** A directory and its associated IDirectoryWatcher::WatchOptions, which represents a unique ID for a watch request */
+	typedef TPair<FString, uint32> FDirectoryWithFlags;
 public:
 	FDirectoryWatcherWindows();
 	virtual ~FDirectoryWatcherWindows();
@@ -21,7 +24,7 @@ public:
 	virtual void Tick (float DeltaSeconds) override;
 
 	/** Map of directory paths to requests */
-	TMap<FString, FDirectoryWatchRequestWindows*> RequestMap;
+	TMap<FDirectoryWithFlags, FDirectoryWatchRequestWindows*> RequestMap;
 	TArray<FDirectoryWatchRequestWindows*> RequestsPendingDelete;
 
 	/** A count of FDirectoryWatchRequestWindows created to ensure they are cleaned up on shutdown */
diff --git a/Engine/Source/Editor/ContentBrowser/Private/SAssetView.cpp b/Engine/Source/Editor/ContentBrowser/Private/SAssetView.cpp
--- a/Engine/Source/Editor/ContentBrowser/Private/SAssetView.cpp
+++ b/Engine/Source/Editor/ContentBrowser/Private/SAssetView.cpp
@@ -2773,18 +2773,19 @@ void SAssetView::OnAssetRegistryPathAdded(const FString& Path)
 		{
 			for (const FName& SourcePathName : SourcesData.PackagePaths)
 			{
-				const FString SourcePath = SourcePathName.ToString();
+				// Ensure that /Folder2 is not considered a subfolder of /Folder by appending /
+				FString SourcePath = SourcePathName.ToString() / TEXT("");
 				if(Path.StartsWith(SourcePath))
 				{
 					const FString SubPath = Path.RightChop(SourcePath.Len());
-					
+
 					TArray<FString> SubPathItemList;
 					SubPath.ParseIntoArray(SubPathItemList, TEXT("/"), /*InCullEmpty=*/true);
 
-					if(SubPathItemList.Num() > 0)
+					if (SubPathItemList.Num() > 0)
 					{
 						const FString NewSubFolder = SourcePath / SubPathItemList[0];
-						if(!Folders.Contains(NewSubFolder))
+						if (!Folders.Contains(NewSubFolder))
 						{
 							FilteredAssetItems.Add(MakeShareable(new FAssetViewFolder(NewSubFolder)));
 							RefreshList();
diff --git a/Engine/Source/Runtime/AssetRegistry/Private/AssetRegistry.cpp b/Engine/Source/Runtime/AssetRegistry/Private/AssetRegistry.cpp
--- a/Engine/Source/Runtime/AssetRegistry/Private/AssetRegistry.cpp
+++ b/Engine/Source/Runtime/AssetRegistry/Private/AssetRegistry.cpp
@@ -135,7 +135,8 @@ UAssetRegistryImpl::UAssetRegistryImpl(const FObjectInitializer& ObjectInitializ
 				if (IFileManager::Get().DirectoryExists(*ContentFolder))
 				{
 					FDelegateHandle NewHandle;
-					DirectoryWatcher->RegisterDirectoryChangedCallback_Handle(ContentFolder, IDirectoryWatcher::FDirectoryChanged::CreateUObject(this, &UAssetRegistryImpl::OnDirectoryChanged), NewHandle);
+					DirectoryWatcher->RegisterDirectoryChangedCallback_Handle(ContentFolder, IDirectoryWatcher::FDirectoryChanged::CreateUObject(this, &UAssetRegistryImpl::OnDirectoryChanged),
+																			  NewHandle, IDirectoryWatcher::WatchOptions::IncludeDirectoryChanges);
 					OnDirectoryChangedDelegateHandles.Add(ContentFolder, NewHandle);
 				}
 			}
@@ -2436,6 +2437,32 @@ void UAssetRegistryImpl::OnDirectoryChanged(const TArray<FFileChangeData>& FileC
 				break;
 			}
 		}
+		else if (bIsValidPackageName)
+		{
+			// This could be a directory or possibly a file with no extension or a wrong extension.
+			// No guaranteed way to know at this point since it may have been deleted.
+			switch (FileChangesProcessed[FileIdx].Action)
+			{
+			case FFileChangeData::FCA_Added:
+			{
+				if (FPaths::DirectoryExists(File))
+				{
+					AddPath(LongPackageName);
+					UE_LOG(LogAssetRegistry, Verbose, TEXT("Directory was added to content directory: %s"), *File);
+					AddPathToSearch(LongPackageName);
+				}
+				break;
+			}
+			case FFileChangeData::FCA_Removed:
+			{
+				RemoveAssetPath(*LongPackageName);
+				UE_LOG(LogAssetRegistry, Verbose, TEXT("Directory was removed from content directory: %s"), *File);
+				break;
+			}
+			default:
+				break;
+			}
+		}
 	}
 
 	if (NewFiles.Num())
@@ -2612,7 +2639,8 @@ void UAssetRegistryImpl::OnContentPathMounted(const FString& InAssetPath, const
 		{
 			// If the path doesn't exist on disk, make it so the watcher will work.
 			IFileManager::Get().MakeDirectory(*FileSystemPath);
-			DirectoryWatcher->RegisterDirectoryChangedCallback_Handle(FileSystemPath, IDirectoryWatcher::FDirectoryChanged::CreateUObject(this, &UAssetRegistryImpl::OnDirectoryChanged), OnContentPathMountedOnDirectoryChangedDelegateHandle);
+			DirectoryWatcher->RegisterDirectoryChangedCallback_Handle(FileSystemPath, IDirectoryWatcher::FDirectoryChanged::CreateUObject(this, &UAssetRegistryImpl::OnDirectoryChanged), 
+																	  OnContentPathMountedOnDirectoryChangedDelegateHandle, IDirectoryWatcher::WatchOptions::IncludeDirectoryChanges);
 		}
 	}
 #endif // WITH_EDITOR
