Remove deleted variable for linux builds

diff --git a/Engine/Source/Developer/DirectoryWatcher/Private/Linux/DirectoryWatchRequestLinux.cpp b/Engine/Source/Developer/DirectoryWatcher/Private/Linux/DirectoryWatchRequestLinux.cpp
--- a/Engine/Source/Developer/DirectoryWatcher/Private/Linux/DirectoryWatchRequestLinux.cpp
+++ b/Engine/Source/Developer/DirectoryWatcher/Private/Linux/DirectoryWatchRequestLinux.cpp
@@ -7,7 +7,6 @@
 FDirectoryWatchRequestLinux::FDirectoryWatchRequestLinux()
 :	bRunning(false)
 ,	bEndWatchRequestInvoked(false)
-,	bIncludeDirectoryChanges(false)
 ,	bWatchSubtree(false)
 {
 	NotifyFilter = IN_CREATE | IN_MOVE | IN_MODIFY | IN_DELETE;
