Linux directory watcher optimizations and fixes for cases where we're running with "virtual" layers in the platform file (i.e. pak files)

diff --git a/Engine/Source/Developer/DirectoryWatcher/Private/Linux/DirectoryWatchRequestLinux.cpp b/Engine/Source/Developer/DirectoryWatcher/Private/Linux/DirectoryWatchRequestLinux.cpp
--- a/Engine/Source/Developer/DirectoryWatcher/Private/Linux/DirectoryWatchRequestLinux.cpp
+++ b/Engine/Source/Developer/DirectoryWatcher/Private/Linux/DirectoryWatchRequestLinux.cpp
@@ -138,7 +138,14 @@ void FDirectoryWatchRequestLinux::WatchDirectoryTree(const FString & RootAbsolut
 	TArray<FString> AllFiles;
 	if (bWatchSubtree)
 	{
-		IFileManager::Get().FindFilesRecursive(AllFiles, *RootAbsolutePath, TEXT("*"), false, true);
+		IPlatformFile::GetPlatformPhysical().IterateDirectoryRecursively(*RootAbsolutePath, [&AllFiles](const TCHAR* Name, bool bIsDirectory)
+			{
+				if (bIsDirectory)
+				{
+					AllFiles.Add(Name);
+				}
+				return true;
+			});
 	}
 	// add the path as well
 	AllFiles.Add(RootAbsolutePath);
@@ -146,23 +153,25 @@ void FDirectoryWatchRequestLinux::WatchDirectoryTree(const FString & RootAbsolut
 	for (int32 FileIdx = 0, NumTotal = AllFiles.Num(); FileIdx < NumTotal; ++FileIdx)
 	{
 		const FString& FolderName = AllFiles[FileIdx];
-		checkf(PathsToWatchDescriptors.Find(FolderName) == nullptr, TEXT("Adding a duplicate watch for directory '%s'"), *FolderName);
 
-		int32 WatchDescriptor = inotify_add_watch(FileDescriptor, TCHAR_TO_UTF8(*FolderName), NotifyFilter);
-		if (WatchDescriptor == -1)
+		if (PathsToWatchDescriptors.Find(FolderName) == nullptr)
 		{
-			int ErrNo = errno;
-			UE_LOG(LogDirectoryWatcher, Warning, TEXT("inotify_add_watch cannot watch folder %s (errno = %d, %s)"), *FolderName,
-				ErrNo,
-				(ErrNo == ENOSPC) ? TEXT("Out of inotify watches, increase user.max_inotify_watches") : UTF8_TO_TCHAR(strerror(ErrNo))
+			int32 WatchDescriptor = inotify_add_watch(FileDescriptor, TCHAR_TO_UTF8(*FolderName), NotifyFilter);
+			if (WatchDescriptor == -1)
+			{
+				int ErrNo = errno;
+				UE_LOG(LogDirectoryWatcher, Warning, TEXT("inotify_add_watch cannot watch folder %s (errno = %d, %s)"), *FolderName,
+					ErrNo,
+					(ErrNo == ENOSPC) ? TEXT("Out of inotify watches, increase user.max_inotify_watches") : UTF8_TO_TCHAR(strerror(ErrNo))
 				);
 			// proceed further
-		}
+			}
 
-		UE_LOG(LogDirectoryWatcher, VeryVerbose, TEXT("+ Added a watch %d for '%s'"), WatchDescriptor, *FolderName);
-		// update the mapping
-		WatchDescriptorsToPaths.Add(WatchDescriptor, FolderName);
-		PathsToWatchDescriptors.Add(FolderName, WatchDescriptor);
+			UE_LOG(LogDirectoryWatcher, VeryVerbose, TEXT("+ Added a watch %d for '%s'"), WatchDescriptor, *FolderName);
+			// update the mapping
+			WatchDescriptorsToPaths.Add(WatchDescriptor, FolderName);
+			PathsToWatchDescriptors.Add(FolderName, WatchDescriptor);
+		}
 	}
 }
 
