Removed references to UE4 in comments

diff --git a/Engine/Source/Developer/DirectoryWatcher/Private/DirectoryWatcherProxy.h b/Engine/Source/Developer/DirectoryWatcher/Private/DirectoryWatcherProxy.h
--- a/Engine/Source/Developer/DirectoryWatcher/Private/DirectoryWatcherProxy.h
+++ b/Engine/Source/Developer/DirectoryWatcher/Private/DirectoryWatcherProxy.h
@@ -23,14 +23,14 @@ public:
 	virtual bool UnregisterDirectoryChangedCallback_Handle(const FString& Directory, FDelegateHandle InHandle) override;
 	virtual void Tick(float DeltaSeconds) override;
 
-	/** Register external changes that the OS file watcher couldn't detect (eg, a file changing in a UE4 sandbox) */
+	/** Register external changes that the OS file watcher couldn't detect (eg, a file changing in a UE sandbox) */
 	void RegisterExternalChanges(TArrayView<const FFileChangeData> FileChanges);
 
 private:
-	/** Register external changes that the OS file watcher couldn't detect (eg, a file changing in a UE4 sandbox) */
+	/** Register external changes that the OS file watcher couldn't detect (eg, a file changing in a UE sandbox) */
 	void RegisterExternalChanges_GameThread(TArrayView<const FFileChangeData> FileChanges);
 
-	/** Process pending external changes that the OS file watcher couldn't detect (eg, a file changing in a UE4 sandbox) */
+	/** Process pending external changes that the OS file watcher couldn't detect (eg, a file changing in a UE sandbox) */
 	void ProcessPendingChanges();
 
 	/** Individual watch callback */
