Fixed missing UE4 comment in DirectoryWatcher

diff --git a/Engine/Source/Developer/DirectoryWatcher/Public/DirectoryWatcherModule.h b/Engine/Source/Developer/DirectoryWatcher/Public/DirectoryWatcherModule.h
--- a/Engine/Source/Developer/DirectoryWatcher/Public/DirectoryWatcherModule.h
+++ b/Engine/Source/Developer/DirectoryWatcher/Public/DirectoryWatcherModule.h
@@ -20,7 +20,7 @@ public:
 	/** Gets the directory watcher singleton or returns NULL if the platform does not support directory watching */
 	virtual IDirectoryWatcher* Get();
 
-	/** Register external changes that the OS file watcher couldn't detect (eg, a file changing in a UE4 sandbox) */
+	/** Register external changes that the OS file watcher couldn't detect (eg, a file changing in a UE sandbox) */
 	virtual void RegisterExternalChanges(TArrayView<const FFileChangeData> FileChanges) const;
 
 private:
