Move error to warning

diff --git a/Engine/Source/Developer/DirectoryWatcher/Private/FileCache.cpp b/Engine/Source/Developer/DirectoryWatcher/Private/FileCache.cpp
--- a/Engine/Source/Developer/DirectoryWatcher/Private/FileCache.cpp
+++ b/Engine/Source/Developer/DirectoryWatcher/Private/FileCache.cpp
@@ -398,7 +398,7 @@ FFileCache::FFileCache(const FFileCacheConfig& InConfig)
 		auto Callback = IDirectoryWatcher::FDirectoryChanged::CreateRaw(this, &FFileCache::OnDirectoryChanged);
 		if (!DirectoryWatcher->RegisterDirectoryChangedCallback_Handle(Config.Directory, Callback, WatcherDelegate))
 		{
-			UE_LOG(LogFileCache, Error, TEXT("Failed registering directory watcher for folder '%s'"), *Config.Directory);
+			UE_LOG(LogFileCache, Warning, TEXT("Failed registering directory watcher for folder '%s'"), *Config.Directory);
 		}
 	}
 }
