Fix FDirectoryWatchRequestLinux::ProcessChanges() to handle UTF8 filenames

ANSI_TO_TCHAR should be UTF8_TO_TCHAR for event name

diff --git a/Engine/Source/Developer/DirectoryWatcher/Private/Linux/DirectoryWatchRequestLinux.cpp b/Engine/Source/Developer/DirectoryWatcher/Private/Linux/DirectoryWatchRequestLinux.cpp
--- a/Engine/Source/Developer/DirectoryWatcher/Private/Linux/DirectoryWatchRequestLinux.cpp
+++ b/Engine/Source/Developer/DirectoryWatcher/Private/Linux/DirectoryWatchRequestLinux.cpp
@@ -264,7 +264,7 @@ void FDirectoryWatchRequestLinux::ProcessChanges()
 				if (EventPathPtr)
 				{
 					FString EventPath = *EventPathPtr;	// get a copy since we can mutate WatchDescriptorToPaths
-					FString AffectedFile = EventPath / ANSI_TO_TCHAR(Event->name);	// by default, some events report about the file itself
+					FString AffectedFile = EventPath / UTF8_TO_TCHAR(Event->name);	// by default, some events report about the file itself
 
 					if ((Event->mask & IN_CREATE) || (Event->mask & IN_MOVED_TO))
 					{
