Clean up Linux inotify warnings and potential leaks

We still have some duplicate inotify watches, but this first pass will spew a lot more information when we hit inotify limits.

Adds a "DumpINotifyStats" command in non-release builds
  Spews global inotify & UE stats, along with physical count of directories, etc.

Canonicalize directory path in FDirectoryWatcherLinux::RegisterDirectoryChangedCallback_Handle
  Shootergame was adding 141 duplicate watches for Samples/Games/ShooterGame/Content w/o this. Was 1136, is now 995.

Change PathsToWatchDescriptors tmap to PathNameHashSet
  Don't need to store full paths for each watch directory twice

Fix bugs in TestPAL in addition to adding DumpStats() command, which looks ~ like this:

LogDirectoryWatcher: Warning: inotify limits
LogDirectoryWatcher: Warning:   max_queued_events: 16384
LogDirectoryWatcher: Warning:   max_user_instances: 128
LogDirectoryWatcher: Warning:   max_user_watches: 65536
LogDirectoryWatcher: Warning: inotify per-process stats
LogDirectoryWatcher: Warning:   systemd (pid 2239) watches:23 instances:3
...
LogDirectoryWatcher: Warning:   plugin_host-3.3 (pid 395041) watches:62 instances:1
LogDirectoryWatcher: Warning:   plugin_host-3.8 (pid 395044) watches:62 instances:1
LogDirectoryWatcher: Warning:   TestPAL (pid 396852) watches:2 instances:1
LogDirectoryWatcher: Warning: Total inotify Watches:392 Instances:28
LogDirectoryWatcher: Warning: Current watch requests
LogDirectoryWatcher: Warning:   /var/tmp/DirectoryWatcherTest396852: 2 watches
LogDirectoryWatcher: Warning:   Total count:2

The above is also dumped (once) when we fail to init or add a inotify watch.

Need to create better documentation and add a pointer to it, similar to what VSCode does: (hat tip Brandon)

    https://code.visualstudio.com/docs/setup/linux#_visual-studio-code-is-unable-to-watch-for-file-changes-in-this-large-workspace-error-enospc

Related bugs:

; FPS BP Cooking Content - errno = 28, Out of inotify watches
https://jira.it.epicgames.com/browse/UE-125210

; inotify Warnings when Cooking Content for Linux
https://jira.it.epicgames.com/browse/UE-119696

; Time Niagara Sequencer failed to play | Error: Couldn't find file for package
https://jira.it.epicgames.com/browse/UE-89750

; inotify warnings from Linux command line builds
https://jira.it.epicgames.com/browse/UE-76562

diff --git a/Engine/Source/Developer/DirectoryWatcher/Private/DirectoryWatcherProxy.cpp b/Engine/Source/Developer/DirectoryWatcher/Private/DirectoryWatcherProxy.cpp
--- a/Engine/Source/Developer/DirectoryWatcher/Private/DirectoryWatcherProxy.cpp
+++ b/Engine/Source/Developer/DirectoryWatcher/Private/DirectoryWatcherProxy.cpp
@@ -66,6 +66,11 @@ void FDirectoryWatcherProxy::Tick(float DeltaSeconds)
 	ProcessPendingChanges();
 }
 
+bool FDirectoryWatcherProxy::DumpStats()
+{
+	return Inner->DumpStats();
+}
+
 void FDirectoryWatcherProxy::RegisterExternalChanges(TArrayView<const FFileChangeData> FileChanges)
 {
 	if (IsInGameThread())
diff --git a/Engine/Source/Developer/DirectoryWatcher/Private/DirectoryWatcherProxy.h b/Engine/Source/Developer/DirectoryWatcher/Private/DirectoryWatcherProxy.h
--- a/Engine/Source/Developer/DirectoryWatcher/Private/DirectoryWatcherProxy.h
+++ b/Engine/Source/Developer/DirectoryWatcher/Private/DirectoryWatcherProxy.h
@@ -22,6 +22,7 @@ public:
 	virtual bool RegisterDirectoryChangedCallback_Handle(const FString& Directory, const FDirectoryChanged& InDelegate, FDelegateHandle& OutHandle, uint32 Flags) override;
 	virtual bool UnregisterDirectoryChangedCallback_Handle(const FString& Directory, FDelegateHandle InHandle) override;
 	virtual void Tick(float DeltaSeconds) override;
+	virtual bool DumpStats() override;
 
 	/** Register external changes that the OS file watcher couldn't detect (eg, a file changing in a UE sandbox) */
 	void RegisterExternalChanges(TArrayView<const FFileChangeData> FileChanges);
diff --git a/Engine/Source/Developer/DirectoryWatcher/Private/Linux/DirectoryWatchRequestLinux.cpp b/Engine/Source/Developer/DirectoryWatcher/Private/Linux/DirectoryWatchRequestLinux.cpp
--- a/Engine/Source/Developer/DirectoryWatcher/Private/Linux/DirectoryWatchRequestLinux.cpp
+++ b/Engine/Source/Developer/DirectoryWatcher/Private/Linux/DirectoryWatchRequestLinux.cpp
@@ -4,12 +4,29 @@
 #include "HAL/FileManager.h"
 #include "DirectoryWatcherPrivate.h"
 
+// To see inotify watch events:
+//   TestPAL dirwatcher -LogCmds="LogDirectoryWatcher VeryVerbose"
+
+#define EVENT_SIZE     ( sizeof(struct inotify_event) )
+#define EVENT_BUF_LEN  ( 1024 * ( EVENT_SIZE + 16 ) )
+
+static bool GDumpStats = false;
+static bool GDumpedError = false;
+static FString GINotifyErrorMsg;
+
+static uint32 GetPathNameHash(const FString& Key)
+{
+	const TCHAR* Str = &Key[0];
+	uint32 StrLen = sizeof(TCHAR) * Key.Len();
+
+	return CityHash64(reinterpret_cast<const char*>(Str), StrLen);
+}
+
 FDirectoryWatchRequestLinux::FDirectoryWatchRequestLinux()
-:	bRunning(false)
+:	bWatchSubtree(false)
 ,	bEndWatchRequestInvoked(false)
-,	bWatchSubtree(false)
+,	FileDescriptor(-1)
 {
-	NotifyFilter = IN_CREATE | IN_MOVE | IN_MODIFY | IN_DELETE;
 }
 
 FDirectoryWatchRequestLinux::~FDirectoryWatchRequestLinux()
@@ -19,10 +36,20 @@ FDirectoryWatchRequestLinux::~FDirectoryWatchRequestLinux()
 
 void FDirectoryWatchRequestLinux::Shutdown()
 {
-	close(FileDescriptor);
+	// Go through and inotify_rm_watch all our watch descriptors
+	for (auto MapIt = WatchDescriptorsToPaths.CreateConstIterator(); MapIt; ++MapIt)
+	{
+		inotify_rm_watch(FileDescriptor, MapIt->Key);
+	}
 
 	WatchDescriptorsToPaths.Empty();
-	PathsToWatchDescriptors.Empty();
+	PathNameHashSet.Empty();
+
+	if (FileDescriptor != -1)
+	{
+		close(FileDescriptor);
+		FileDescriptor = -1;
+	}
 }
 
 bool FDirectoryWatchRequestLinux::Init(const FString& InDirectory, uint32 Flags)
@@ -33,42 +60,33 @@ bool FDirectoryWatchRequestLinux::Init(const FString& InDirectory, uint32 Flags)
 		return false;
 	}
 
-	Directory = InDirectory;
-	bWatchSubtree = (Flags & IDirectoryWatcher::WatchOptions::IgnoreChangesInSubtree) == 0;
+	Shutdown();
 
-	if (bRunning)
-	{
-		Shutdown();
-	}
+	bWatchSubtree = (Flags & IDirectoryWatcher::WatchOptions::IgnoreChangesInSubtree) == 0;
 
 	bEndWatchRequestInvoked = false;
 
 	// Make sure the path is absolute
-	const FString FullPath = FPaths::ConvertRelativePathToFull(InDirectory);
-	UE_LOG(LogDirectoryWatcher, Verbose, TEXT("Adding watch for directory tree '%s'"), *FullPath);
+	WatchDirectory = FPaths::ConvertRelativePathToFull(InDirectory);
 
-	FileDescriptor = inotify_init1(IN_NONBLOCK | IN_CLOEXEC);
+	UE_LOG(LogDirectoryWatcher, Verbose, TEXT("Adding watch for directory tree '%s'"), *WatchDirectory);
 
+	FileDescriptor = inotify_init1(IN_NONBLOCK | IN_CLOEXEC);
 	if (FileDescriptor == -1)
 	{
-		// Failed to init inotify. The number of inotify instances on the default Linux is vastly
-		// insufficient, so do not make this an error which can break cooking
-		int ErrNo = errno;
-		if (ErrNo == EMFILE)
+		if (errno == EMFILE)
 		{
-			UE_LOG(LogDirectoryWatcher, Warning, TEXT("Failed to init inotify (ran out of inotify instances, consider increasing system limits)"));
+			SetINotifyErrorMsg(TEXT("Failed to init inotify (ran out of inotify instances)"));
 		}
 		else
 		{
-			UE_LOG(LogDirectoryWatcher, Error, TEXT("Failed to init inotify (errno=%d, %s)"), ErrNo, UTF8_TO_TCHAR(strerror(ErrNo)));
+			UE_LOG(LogDirectoryWatcher, Error, TEXT("Failed to init inotify (errno=%d, %s)"), errno, UTF8_TO_TCHAR(strerror(errno)));
 		}
 		return false;
 	}
 
 	// find all subdirs
-	WatchDirectoryTree(FullPath);
-
-	bRunning = true;
+	WatchDirectoryTree(WatchDirectory);
 
 	return true;
 }
@@ -96,19 +114,42 @@ void FDirectoryWatchRequestLinux::EndWatchRequest()
 	bEndWatchRequestInvoked = true;
 }
 
+void FDirectoryWatchRequestLinux::ProcessNotifications(TMap<FString, FDirectoryWatchRequestLinux*>& RequestMap)
+{
+	// Trigger any file change notification delegates
+	for (auto MapIt = RequestMap.CreateConstIterator(); MapIt; ++MapIt)
+	{
+		FDirectoryWatchRequestLinux &WatchRequest = *MapIt.Value();
+
+		WatchRequest.ProcessPendingNotifications();
+	}
+
+	DumpINotifyErrorDetails(RequestMap);
+}
+
+void FDirectoryWatchRequestLinux::DumpStats(TMap<FString, FDirectoryWatchRequestLinux*>& RequestMap)
+{
+	GDumpStats = true;
+	DumpINotifyErrorDetails(RequestMap);
+}
 
 void FDirectoryWatchRequestLinux::ProcessPendingNotifications()
 {
-	ProcessChanges();
+	/** Each FFileChangeData tracks whether it is a directory or not */
+	TArray<TPair<FFileChangeData, bool>> FileChanges;
+
+	ProcessChanges(FileChanges);
 
 	// Trigger all listening delegates with the files that have changed
 	if (FileChanges.Num() > 0)
 	{
 		TMap<uint32, TArray<FFileChangeData>> FileChangeCache;
+
 		for (const FWatchDelegate& Delegate : Delegates)
 		{
 			// Filter list of all file changes down to ones that just match this delegate's flags
 			TArray<FFileChangeData>* CachedChanges = FileChangeCache.Find(Delegate.Value);
+
 			if (CachedChanges)
 			{
 				Delegate.Key.Execute(*CachedChanges);
@@ -117,6 +158,7 @@ void FDirectoryWatchRequestLinux::ProcessPendingNotifications()
 			{
 				const bool bIncludeDirs = (Delegate.Value & IDirectoryWatcher::WatchOptions::IncludeDirectoryChanges) != 0;
 				TArray<FFileChangeData>& Changes = FileChangeCache.Add(Delegate.Value);
+
 				for (const TPair<FFileChangeData, bool>& FileChangeData : FileChanges)
 				{
 					if (!FileChangeData.Value || bIncludeDirs)
@@ -134,110 +176,158 @@ void FDirectoryWatchRequestLinux::ProcessPendingNotifications()
 
 void FDirectoryWatchRequestLinux::WatchDirectoryTree(const FString & RootAbsolutePath)
 {
+	if (bEndWatchRequestInvoked || (FileDescriptor == -1))
+	{
+		return;
+	}
+
+	// If this isn't our root watch directory or under it, don't watch
+	if (!RootAbsolutePath.StartsWith(WatchDirectory, ESearchCase::CaseSensitive) ||
+		(!bWatchSubtree && (RootAbsolutePath != WatchDirectory)))
+	{
+		return;
+	}
+
 	UE_LOG(LogDirectoryWatcher, VeryVerbose, TEXT("Watching tree '%s'"), *RootAbsolutePath);
+
 	TArray<FString> AllFiles;
 	if (bWatchSubtree)
 	{
-		IPlatformFile::GetPlatformPhysical().IterateDirectoryRecursively(*RootAbsolutePath, [&AllFiles](const TCHAR* Name, bool bIsDirectory)
-			{
-				if (bIsDirectory)
+		IPlatformFile::GetPlatformPhysical().IterateDirectoryRecursively(*RootAbsolutePath,
+			[&AllFiles](const TCHAR* Name, bool bIsDirectory)
 				{
-					AllFiles.Add(Name);
-				}
-				return true;
-			});
+					if (bIsDirectory)
+					{
+						AllFiles.Add(Name);
+					}
+					return true;
+				});
 	}
-	// add the path as well
+
+	// Add root path
 	AllFiles.Add(RootAbsolutePath);
 
-	for (int32 FileIdx = 0, NumTotal = AllFiles.Num(); FileIdx < NumTotal; ++FileIdx)
+	for (const FString& FolderName: AllFiles)
 	{
-		const FString& FolderName = AllFiles[FileIdx];
+		uint32 PathNameHash = GetPathNameHash(FolderName);
 
-		if (PathsToWatchDescriptors.Find(FolderName) == nullptr)
+		if (PathNameHashSet.Find(PathNameHash) == nullptr)
 		{
+			int32 NotifyFilter = IN_CREATE | IN_MOVE | IN_MODIFY | IN_DELETE;
 			int32 WatchDescriptor = inotify_add_watch(FileDescriptor, TCHAR_TO_UTF8(*FolderName), NotifyFilter);
+
 			if (WatchDescriptor == -1)
 			{
-				int ErrNo = errno;
-				UE_LOG(LogDirectoryWatcher, Warning, TEXT("inotify_add_watch cannot watch folder %s (errno = %d, %s)"), *FolderName,
-					ErrNo,
-					(ErrNo == ENOSPC) ? TEXT("Out of inotify watches, increase user.max_inotify_watches") : UTF8_TO_TCHAR(strerror(ErrNo))
-				);
-			// proceed further
+				if (errno == ENOSPC)
+				{
+					FString ErrorMsg = FString::Printf(
+						TEXT("inotify_add_watch cannot watch folder %s (Out of inotify watches)"), *FolderName);
+					SetINotifyErrorMsg(ErrorMsg);
+				}
+				else
+				{
+					UE_LOG(LogDirectoryWatcher, Warning, TEXT("inotify_add_watch cannot watch folder %s (errno = %d, %s)"),
+							*FolderName, errno, UTF8_TO_TCHAR(strerror(errno)));
+				}
 			}
+			else
+			{
+				UE_LOG(LogDirectoryWatcher, VeryVerbose, TEXT("+ Added a watch %d for '%s'"), WatchDescriptor, *FolderName);
+
+				// Set the inotify watch descriptor -> folder name mapping
+				WatchDescriptorsToPaths.Add(WatchDescriptor, FolderName);
 
-			UE_LOG(LogDirectoryWatcher, VeryVerbose, TEXT("+ Added a watch %d for '%s'"), WatchDescriptor, *FolderName);
-			// update the mapping
-			WatchDescriptorsToPaths.Add(WatchDescriptor, FolderName);
-			PathsToWatchDescriptors.Add(FolderName, WatchDescriptor);
+				// Add hashed directory path
+				PathNameHashSet.Add(PathNameHash);
+			}
 		}
 	}
 }
 
-void FDirectoryWatchRequestLinux::UnwatchDirectoryTree(const FString & RootAbsolutePath)
+void FDirectoryWatchRequestLinux::UnwatchDirectoryTree(const FString& RootAbsolutePath)
 {
 	UE_LOG(LogDirectoryWatcher, VeryVerbose, TEXT("Unwatching tree '%s'"), *RootAbsolutePath);
-	// remove the watch for the folder and all subfolders
-	// since it is expected that there will be a lot of them, just create a new TMap
-	TMap<FString, int32, FDefaultSetAllocator, FCaseSensitiveLookupKeyFuncs<int32>> NewPathsToWatchDescriptors;
-	for (auto MapIt = PathsToWatchDescriptors.CreateIterator(); MapIt; ++MapIt)
+
+	for (auto MapIt = WatchDescriptorsToPaths.CreateIterator(); MapIt; ++MapIt)
 	{
-		if (!MapIt->Key.StartsWith(RootAbsolutePath, ESearchCase::CaseSensitive))
-		{
-			NewPathsToWatchDescriptors.Add(MapIt->Key, MapIt->Value);
-		}
-		else
+		int WatchDescriptor = MapIt->Key;
+		const FString& Path = MapIt->Value;
+
+		if (Path.StartsWith(RootAbsolutePath, ESearchCase::CaseSensitive))
 		{
-			UE_LOG(LogDirectoryWatcher, VeryVerbose, TEXT("- Removing a watch %d for '%s'"), MapIt->Value, *MapIt->Key);
+			UE_LOG(LogDirectoryWatcher, VeryVerbose, TEXT("- Removing a watch %d for '%s'"), WatchDescriptor, *Path);
 
 			// delete the descriptor
-			int RetVal = inotify_rm_watch(FileDescriptor, MapIt->Value);
+			int RetVal = inotify_rm_watch(FileDescriptor, WatchDescriptor);
 
-			// why check for RootAbsolutePath? Because this function may be called when root path has been deleted, and inotify_rm_watch() will fail
-			// removing a watch on a deleted file... yay for API symmetry. Just "leak" the watch descriptor without the warning
-			if (RetVal == -1 || MapIt->Key != RootAbsolutePath)
+			// This function may be called when root path has been deleted, and inotify_rm_watch() will fail
+			// with an EINVAL when removing a watch on a deleted file.
+			if (RetVal == -1 && errno != EINVAL)
 			{
-				int ErrNo = errno;
-				UE_LOG(LogDirectoryWatcher, Error, TEXT("inotify_rm_watch cannot remove descriptor %d for folder '%s' (errno = %d, %s)"),
-					MapIt->Value,
-					*MapIt->Key,
-					ErrNo,
-					ANSI_TO_TCHAR(strerror(ErrNo))
-				);
+				UE_LOG(LogDirectoryWatcher, Warning, TEXT("inotify_rm_watch cannot remove descriptor %d for folder '%s' (errno = %d, %s)"),
+						WatchDescriptor, *Path, errno, ANSI_TO_TCHAR(strerror(errno)));
 			}
-			WatchDescriptorsToPaths.Remove(MapIt->Value);
+
+			// Safe version of:
+			//   WatchDescriptorsToPaths.Remove(WatchDescriptor);
+			MapIt.RemoveCurrent();
+
+			PathNameHashSet.Remove(GetPathNameHash(Path));
 		}
 	}
+}
 
-	PathsToWatchDescriptors = NewPathsToWatchDescriptors;
+static FString INotifyFlagsToStr(uint32 INotifyFlags)
+{
+#if UE_BUILD_SHIPPING
+	return FString();
+#else
+	FString Ret = TEXT("[");
+
+#define _XTAG(_x) if (INotifyFlags & _x) Ret += FString(TEXT(" ")) + TEXT(#_x)
+	_XTAG(IN_ACCESS);
+	_XTAG(IN_MODIFY);
+	_XTAG(IN_ATTRIB);
+	_XTAG(IN_CLOSE_WRITE);
+	_XTAG(IN_CLOSE_NOWRITE);
+	_XTAG(IN_OPEN);
+	_XTAG(IN_MOVED_FROM);
+	_XTAG(IN_MOVED_TO);
+	_XTAG(IN_CREATE);
+	_XTAG(IN_DELETE);
+	_XTAG(IN_DELETE_SELF);
+	_XTAG(IN_MOVE_SELF);
+	_XTAG(IN_UNMOUNT);
+	_XTAG(IN_Q_OVERFLOW);
+	_XTAG(IN_IGNORED);
+	_XTAG(IN_ISDIR);
+#undef _XTAG
+
+	Ret += TEXT(" ]");
+	return Ret;
+#endif
 }
 
-void FDirectoryWatchRequestLinux::ProcessChanges()
+void FDirectoryWatchRequestLinux::ProcessChanges(TArray<TPair<FFileChangeData, bool>>& FileChanges)
 {
 	uint8_t Buffer[EVENT_BUF_LEN] __attribute__ ((aligned(__alignof__(struct inotify_event))));
-	const struct inotify_event* Event;
-	ssize_t Len = 0;
-	uint8_t* Ptr = nullptr;
+
+	if (FileDescriptor == -1)
+	{
+		return;
+	}
 
 	// Loop while events can be read from inotify file descriptor
 	for (;;)
 	{
-		// default action
-		FFileChangeData::EFileChangeAction Action = FFileChangeData::FCA_Unknown;
-
 		// Read event stream
-		Len = read(FileDescriptor, Buffer, EVENT_BUF_LEN);
+		ssize_t Len = read(FileDescriptor, Buffer, EVENT_BUF_LEN);
 
 		// If the non-blocking read() found no events to read, then it returns -1 with errno set to EAGAIN.
 		if (Len == -1 && errno != EAGAIN)
 		{
-			int ErrNo = errno;
 			UE_LOG(LogDirectoryWatcher, Error, TEXT("FDirectoryWatchRequestLinux::ProcessChanges() read() error getting events for path '%s' (errno = %d, %s)"),
-				*Directory,
-				ErrNo,
-				ANSI_TO_TCHAR(strerror(ErrNo))
-				);
+				*WatchDirectory, errno, ANSI_TO_TCHAR(strerror(errno)));
 			break;
 		}
 
@@ -246,31 +336,45 @@ void FDirectoryWatchRequestLinux::ProcessChanges()
 			break;
 		}
 
+		if (bEndWatchRequestInvoked)
+		{
+			continue;
+		}
+
 		// Loop over all events in the buffer
-		for (Ptr = Buffer; Ptr < Buffer + Len;)
+		uint8_t* Ptr = Buffer;
+		while (Ptr < Buffer + Len)
 		{
+			const struct inotify_event* Event;
+			FFileChangeData::EFileChangeAction Action = FFileChangeData::FCA_Unknown;
+
 			Event = reinterpret_cast<const struct inotify_event *>(Ptr);
+			Ptr += EVENT_SIZE + Event->len;
 
-			// skip if overflowed
-			if (Event->wd != -1 && (Event->mask & IN_Q_OVERFLOW) == 0)
+			// Skip if overflowed
+			if ((Event->wd != -1) && (Event->mask & IN_Q_OVERFLOW) == 0)
 			{
 				const FString *EventPathPtr = WatchDescriptorsToPaths.Find(Event->wd);
 
-				UE_LOG(LogDirectoryWatcher, VeryVerbose, TEXT("Event: watch descriptor %d, mask 0x%08x, EventPath: '%s'"),
-					Event->wd, Event->mask, EventPathPtr ? *(*EventPathPtr) : TEXT("nullptr"), ANSI_TO_TCHAR(Event->name));
+				UE_LOG(LogDirectoryWatcher, VeryVerbose, TEXT("Event: watch descriptor %d, mask 0x%08x, EventPath: '%s' Event Name: '%s' %s"),
+					Event->wd, Event->mask, EventPathPtr ? *(*EventPathPtr) : TEXT("nullptr"), ANSI_TO_TCHAR(Event->name), *INotifyFlagsToStr(Event->mask));
 
-				// if we're geting multiple events (e.g. DELETE, IGNORED) we could ahve removed descriptor on previous iteration,
+				// If we're geting multiple events (e.g. DELETE, IGNORED) we could have removed descriptor on previous iteration,
 				// so we need to handle inability to find it in the map
 				if (EventPathPtr)
 				{
-					FString EventPath = *EventPathPtr;	// get a copy since we can mutate WatchDescriptorToPaths
-					FString AffectedFile = EventPath / UTF8_TO_TCHAR(Event->name);	// by default, some events report about the file itself
+					const FString& EventPath = *EventPathPtr;
+					FString AffectedFile = EventPath / UTF8_TO_TCHAR(Event->name); // By default, some events report about the file itself
+
+					bool bIsDir = (Event->mask & IN_ISDIR) != 0;
 
 					if ((Event->mask & IN_CREATE) || (Event->mask & IN_MOVED_TO))
 					{
-						// if a directory was created/moved, watch it
-						if (Event->mask & IN_ISDIR)
+						// IN_CREATE: File/directory created in watched directory
+						// IN_MOVED_TO: Generated for the directory containing the new filename when a file is renamed
+						if (bIsDir)
 						{
+							// If a directory was created/moved, watch it
 							WatchDirectoryTree(AffectedFile);
 						}
 
@@ -278,33 +382,49 @@ void FDirectoryWatchRequestLinux::ProcessChanges()
 					}
 					else if (Event->mask & IN_MODIFY)
 					{
-						// if a directory was modified, we expect to get events from already watched files in it
+						// IN_MODIFY: File was modified
+						// If a directory was modified, we expect to get events from already watched files in it
 						Action = FFileChangeData::FCA_Modified;
 					}
-					// check if the file/directory itself has been deleted (IGNORED can also be sent on delete)
+					// Check if the file/directory itself has been deleted (IGNORED can also be sent on delete)
 					else if ((Event->mask & IN_DELETE_SELF) || (Event->mask & IN_IGNORED) || (Event->mask & IN_UNMOUNT))
 					{
-						// if a directory was deleted, we expect to get events from already watched files in it
+						// IN_DELETE_SELF: Watched file/directory was itself deleted.
+						//   In addition, an IN_IGNORED event will subsequently be generated for the watch descriptor
+						// IN_UNMOUNT: Filesystem containing watched object was unmounted.
+						//   In addition, an IN_IGNORED event will subsequently be generated for the watch descriptor
+						// IN_IGNORED: Watch was removed explicitly (inotify_rm_watch(2)) or
+						//   automatically (file was deleted, or filesystem was unmounted).
+
+						// If a directory was deleted, we expect to get events from already watched files in it
 						AffectedFile = EventPath;
-						if (Event->mask & IN_ISDIR)
+
+						if (bIsDir)
 						{
 							UnwatchDirectoryTree(EventPath);
 						}
 						else
 						{
-							// remove itself from both mappings
-							// NOTE: inotify_rm_watch() will fail here as watch descriptor is no longer valid (deficiency of inotify?).
-							// Just avoid calling inotify_rm_watch() even if it feels like a leak
+							// NOTE: I don't *believe* this code should ever get called? We only watch directories.
+							// But to be on the safe side, still handling this case...
+
+							// inotify_rm_watch() may fail here as watch descriptor is potentially no longer valid, so ignore return
+							inotify_rm_watch(FileDescriptor, Event->wd);
+
+							// Remove from both mappings
 							WatchDescriptorsToPaths.Remove(Event->wd);
-							PathsToWatchDescriptors.Remove(EventPath);
+							PathNameHashSet.Remove(GetPathNameHash(EventPath));
 						}
 
 						Action = FFileChangeData::FCA_Removed;
 					}
 					else if ((Event->mask & IN_DELETE) || (Event->mask & IN_MOVED_FROM))
 					{
-						// if a directory was deleted/moved, invalidate watch descriptors associated with it (unwatch it)
-						if (Event->mask & IN_ISDIR)
+						// IN_DELETE: File/directory deleted from watched directory
+						// IN_MOVED_FROM: Generated for the directory containing the old filename when a file is renamed
+
+						// If a directory was deleted/moved, unwatch it
+						if (bIsDir)
 						{
 							UnwatchDirectoryTree(AffectedFile);
 						}
@@ -314,12 +434,245 @@ void FDirectoryWatchRequestLinux::ProcessChanges()
 
 					if (Event->len)
 					{
-						FileChanges.Emplace(FFileChangeData(AffectedFile, Action), (Event->mask & IN_ISDIR) != 0);
+						FileChanges.Emplace(FFileChangeData(AffectedFile, Action), bIsDir);
 					}
 				}
 			}
+		}
+	}
+}
 
-			Ptr += EVENT_SIZE + Event->len;
+static uint32 get_inotify_procfs_value(const char *FileName)
+{
+	char Buf[256];
+	uint32 InterfaceVal = 0;
+	
+	snprintf(Buf, sizeof(Buf), "/proc/sys/fs/inotify/%s", FileName);
+	Buf[sizeof(Buf) - 1] = 0;
+	
+	FILE* FilePtr = fopen(Buf, "r");
+	if (FilePtr)
+	{
+		if (fscanf(FilePtr, "%u", &InterfaceVal) != 1)
+		{
+			InterfaceVal = 0;
 		}
+		fclose(FilePtr);
 	}
+
+	return InterfaceVal;
 }
+
+void FDirectoryWatchRequestLinux::SetINotifyErrorMsg(const FString &ErrorMsg)
+{
+	if (!GINotifyErrorMsg.Len())
+	{
+		GINotifyErrorMsg = ErrorMsg;
+	}
+}
+
+static FString GetLinkName(const char *Pathname)
+{
+	FString Result;
+	char Filename[PATH_MAX + 1];
+
+	ssize_t Ret = readlink(Pathname, Filename, sizeof(Filename));
+	if ((Ret > 0) && (Ret < sizeof(Filename)))
+	{
+		Filename[Ret] = 0;
+		Result = Filename;
+	}
+	return Result;
+}
+
+static uint32 INotifyParseFDInfoFile(const FString& Executable, int Pid, const char *d_name)
+{
+	uint32 INotifyCount = 0;
+
+	FILE* FilePtr = fopen(d_name, "r");
+	if (FilePtr)
+	{
+		char line_buf[256];
+
+		for (;;)
+		{
+			if (!fgets(line_buf, sizeof(line_buf), FilePtr))
+			{
+				break;
+			}
+
+			if (!strncmp(line_buf, "inotify ", 8))
+			{
+				INotifyCount++;
+			}
+		}
+
+		fclose(FilePtr);
+	}
+
+	return INotifyCount;
+}
+
+static void INotifyParseFDDir(const FString& Executable, int Pid, uint32 &INotifyCountTotal, uint32& INotifyInstancesTotal)
+{
+	char Buf[256];
+	uint32 INotifyCount = 0;
+	uint32 INotifyInstances = 0;
+
+	snprintf(Buf, sizeof(Buf), "/proc/%d/fd", Pid);
+	Buf[sizeof(Buf) - 1] = 0;
+
+	DIR* dir_fd = opendir(Buf);
+	if (dir_fd)
+	{
+		for (;;)
+		{
+			struct dirent* dp_fd = readdir(dir_fd);
+			if (!dp_fd)
+			{
+				break;
+			}
+
+			if ((dp_fd->d_type == DT_LNK) && isdigit(dp_fd->d_name[0]))
+			{
+				snprintf(Buf, sizeof(Buf), "/proc/%d/fd/%s", Pid, dp_fd->d_name);
+				Buf[sizeof(Buf) - 1] = 0;
+
+				FString Filename = GetLinkName(Buf);
+				if (Filename == TEXT("anon_inode:inotify"))
+				{
+					snprintf(Buf, sizeof(Buf), "/proc/%d/fdinfo/%s", Pid, dp_fd->d_name);
+					Buf[sizeof(Buf) - 1] = 0;
+
+					uint32 Count = INotifyParseFDInfoFile(Executable, Pid, Buf);
+					if (Count)
+					{
+						INotifyInstances++;
+						INotifyCount += Count;
+					}
+				}
+			}
+		}
+
+		closedir(dir_fd);
+	}
+
+	if (INotifyCount)
+	{
+		FString ExeName = FPaths::GetCleanFilename(Executable);
+
+		UE_LOG(LogDirectoryWatcher, Warning, TEXT("  %s (pid %d) watches:%u instances:%u"), *ExeName, Pid, INotifyCount, INotifyInstances);
+
+		INotifyCountTotal += INotifyCount;
+		INotifyInstancesTotal += INotifyInstances;
+	}
+}
+
+static void INotifyDumpProcessStats()
+{
+	uint32 INotifyCountTotal = 0;
+	uint32 INotifyInstancesTotal = 0;
+
+	DIR* dir_proc = opendir("/proc");
+	if (dir_proc)
+	{
+		for (;;)
+		{
+			struct dirent* dp_proc = readdir(dir_proc);
+			if (!dp_proc)
+			{
+				break;
+			}
+
+			if ((dp_proc->d_type == DT_DIR) && isdigit(dp_proc->d_name[0]))
+			{
+				char Buf[256];
+				int32 Pid = atoi(dp_proc->d_name);
+
+				snprintf(Buf, sizeof(Buf), "/proc/%d/exe", Pid);
+				Buf[sizeof(Buf) - 1] = 0;
+
+				FString Executable = GetLinkName(Buf);
+				if (Executable.Len())
+				{
+					INotifyParseFDDir(Executable, Pid, INotifyCountTotal, INotifyInstancesTotal);
+				}
+			}
+		}
+
+		closedir(dir_proc);
+	}
+
+	UE_LOG(LogDirectoryWatcher, Warning, TEXT("Total inotify Watches:%u Instances:%u"), INotifyCountTotal, INotifyInstancesTotal);
+}
+
+void FDirectoryWatchRequestLinux::DumpINotifyErrorDetails(TMap<FString, FDirectoryWatchRequestLinux*>& RequestMap)
+{
+	if (!GDumpStats)
+	{
+		if (GDumpedError || !GINotifyErrorMsg.Len())
+		{
+			return;
+		}
+		GDumpedError = true;
+
+		UE_LOG(LogDirectoryWatcher, Warning, TEXT("inotify error: %s"), *GINotifyErrorMsg);
+	}
+	GDumpStats = false;
+
+	uint32 MaxQueuedEvents = get_inotify_procfs_value("max_queued_events");
+	uint32 MaxUserInstances = get_inotify_procfs_value("max_user_instances");
+	uint32 MaxUserWatches = get_inotify_procfs_value("max_user_watches");
+
+	UE_LOG(LogDirectoryWatcher, Warning, TEXT("inotify limits"));
+	UE_LOG(LogDirectoryWatcher, Warning, TEXT("  max_queued_events: %u"), MaxQueuedEvents);
+	UE_LOG(LogDirectoryWatcher, Warning, TEXT("  max_user_instances: %u"), MaxUserInstances);
+	UE_LOG(LogDirectoryWatcher, Warning, TEXT("  max_user_watches: %u"), MaxUserWatches);
+
+	UE_LOG(LogDirectoryWatcher, Warning, TEXT("inotify per-process stats"));
+	INotifyDumpProcessStats();
+
+	uint32 Count = 0;
+	UE_LOG(LogDirectoryWatcher, Warning, TEXT("Current watch requests"));
+
+	for (auto MapIt = RequestMap.CreateConstIterator(); MapIt; ++MapIt)
+	{
+		uint32 DirCount = 1;
+		FDirectoryWatchRequestLinux &WatchRequest = *MapIt.Value();
+
+		// Get actual count of subdirectories
+		if (WatchRequest.bWatchSubtree)
+		{
+			IPlatformFile::GetPlatformPhysical().IterateDirectoryRecursively(*WatchRequest.WatchDirectory,
+				[&DirCount](const TCHAR* Name, bool bIsDirectory)
+				{
+					DirCount += bIsDirectory;
+					return true;
+				});
+		}
+
+		UE_LOG(LogDirectoryWatcher, Warning, TEXT("  %s: %u watches (%u total dirs)"),
+			*WatchRequest.WatchDirectory, WatchRequest.WatchDescriptorsToPaths.Num(), DirCount);
+		
+		Count += WatchRequest.WatchDescriptorsToPaths.Num();
+	}
+	UE_LOG(LogDirectoryWatcher, Warning, TEXT("Total UE inotify Watches:%u Instances:%u"), Count, RequestMap.Num());
+}
+
+#if !UE_BUILD_SHIPPING
+
+static bool INotifyCommandHandler(UWorld* InWorld, const TCHAR* Cmd, FOutputDevice& Ar)
+{
+	if (FParse::Command(&Cmd, TEXT("DumpINotifyStats")))
+	{
+		UE_LOG(LogDirectoryWatcher, Warning, TEXT("Dumping inotify stats"));
+		GDumpStats = true;
+		return true;
+	}
+
+	return false;
+}
+
+FStaticSelfRegisteringExec FDirectoryWatchRequestLinuxExecs(INotifyCommandHandler);
+
+#endif
diff --git a/Engine/Source/Developer/DirectoryWatcher/Private/Linux/DirectoryWatchRequestLinux.h b/Engine/Source/Developer/DirectoryWatcher/Private/Linux/DirectoryWatchRequestLinux.h
--- a/Engine/Source/Developer/DirectoryWatcher/Private/Linux/DirectoryWatchRequestLinux.h
+++ b/Engine/Source/Developer/DirectoryWatcher/Private/Linux/DirectoryWatchRequestLinux.h
@@ -6,34 +6,10 @@
 #include "IDirectoryWatcher.h"
 #include <sys/inotify.h>
 
-#define EVENT_SIZE  ( sizeof (struct inotify_event) )
-#define EVENT_BUF_LEN     ( 1024 * ( EVENT_SIZE + 16 ) )
-
 class FDirectoryWatchRequestLinux
 {
 public:
 
-	/**
-	 * Template class to support FString to TValueType maps with a case sensitive key
-	 */
-	template<typename TValueType>
-	struct FCaseSensitiveLookupKeyFuncs : BaseKeyFuncs<TValueType, FString>
-	{
-		static FORCEINLINE const FString& GetSetKey(const TPair<FString, TValueType>& Element)
-		{
-			return Element.Key;
-		}
-		static FORCEINLINE bool Matches(const FString& A, const FString& B)
-		{
-			return A.Equals(B, ESearchCase::CaseSensitive);
-		}
-		static FORCEINLINE uint32 GetKeyHash(const FString& Key)
-		{
-			return FCrc::StrCrc32<TCHAR>(*Key);
-		}
-	};
-
-
 	FDirectoryWatchRequestLinux();
 	virtual ~FDirectoryWatchRequestLinux();
 
@@ -48,41 +24,52 @@ public:
 	bool HasDelegates() const;
 	/** Prepares the request for deletion */
 	void EndWatchRequest();
-	/** Triggers all pending file change notifications */
-	void ProcessPendingNotifications();
+
+	/** Call ProcessPendingNotifications on each delegate */
+	static void ProcessNotifications(TMap<FString, FDirectoryWatchRequestLinux*>& RequestMap);
+
+	/** Dump inotify stats */
+	static void DumpStats(TMap<FString, FDirectoryWatchRequestLinux*>& RequestMap);
 
 private:
 
+	/** Triggers all pending file change notifications */
+	void ProcessPendingNotifications();
+
 	/** Adds watches for all files (and subdirectories) in a directory. */
-	void WatchDirectoryTree(const FString & RootAbsolutePath);
+	void WatchDirectoryTree(const FString& RootAbsolutePath);
 
-	/** Removes all watches from a */
-	void UnwatchDirectoryTree(const FString & RootAbsolutePath);
+	/** Removes all watches for path */
+	void UnwatchDirectoryTree(const FString& RootAbsolutePath);
 
-	FString Directory;
+	void Shutdown();
+	void ProcessChanges(TArray<TPair<FFileChangeData, bool>>& FileChanges);
 
-	bool bRunning;
-	bool bEndWatchRequestInvoked;
+	static void SetINotifyErrorMsg(const FString &ErrorMsg);
+	static void DumpINotifyErrorDetails(TMap<FString, FDirectoryWatchRequestLinux*>& RequestMap);
+
+private:
 
 	/** Whether or not watch subtree. */
 	bool bWatchSubtree;
 
+	/** EndWatchRequest called? */
+	bool bEndWatchRequestInvoked;
+
+	/** INotify file descriptor */
 	int FileDescriptor;
 
-	/** Mapping from watch descriptors to their path names */
-	TMap<int32, FString> WatchDescriptorsToPaths;
+	/** Absolute path to our root watch directory */
+	FString WatchDirectory;
 
-	/** Mapping from paths to watch descriptors */
-	TMap<FString, int32, FDefaultSetAllocator, FCaseSensitiveLookupKeyFuncs<int32>> PathsToWatchDescriptors;
+	/** Mapping from watch descriptors to watched directory names */
+	TMap<int32, FString> WatchDescriptorsToPaths;
 
-	int NotifyFilter;
+	/** Set of hashed directory names we're watching */
+	TSet<uint32> PathNameHashSet;
 
 	/** A delegate with its corresponding IDirectoryWatcher::WatchOptions flags */
 	typedef TPair<IDirectoryWatcher::FDirectoryChanged, uint32> FWatchDelegate;
 	TArray<FWatchDelegate> Delegates;
-	/** Each FFileChangeData tracks whether it is a directory or not */
-	TArray<TPair<FFileChangeData, bool>> FileChanges;
-
-	void Shutdown();
-	void ProcessChanges();
 };
+
diff --git a/Engine/Source/Developer/DirectoryWatcher/Private/Linux/DirectoryWatcherLinux.cpp b/Engine/Source/Developer/DirectoryWatcher/Private/Linux/DirectoryWatcherLinux.cpp
--- a/Engine/Source/Developer/DirectoryWatcher/Private/Linux/DirectoryWatcherLinux.cpp
+++ b/Engine/Source/Developer/DirectoryWatcher/Private/Linux/DirectoryWatcherLinux.cpp
@@ -38,9 +38,11 @@ FDirectoryWatcherLinux::~FDirectoryWatcherLinux()
 	ensure(NumRequests == 0);
 }
 
-bool FDirectoryWatcherLinux::RegisterDirectoryChangedCallback_Handle(const FString& Directory, const FDirectoryChanged& InDelegate, FDelegateHandle& OutHandle, uint32 Flags)
+bool FDirectoryWatcherLinux::RegisterDirectoryChangedCallback_Handle(const FString& InDirectory, const FDirectoryChanged& InDelegate, FDelegateHandle& OutHandle, uint32 Flags)
 {
-	FDirectoryWatchRequestLinux** RequestPtr = RequestMap.Find(Directory);
+	// Make sure the path is absolute. Without this, we were missing and duplicating some directories.
+	FString WatchDirectory = FPaths::ConvertRelativePathToFull(InDirectory);
+	FDirectoryWatchRequestLinux** RequestPtr = RequestMap.Find(WatchDirectory);
 	FDirectoryWatchRequestLinux* Request = NULL;
 
 	if (RequestPtr)
@@ -56,15 +58,15 @@ bool FDirectoryWatcherLinux::RegisterDirectoryChangedCallback_Handle(const FStri
 		NumRequests++;
 
 		// Begin reading directory changes
-		if (!Request->Init(Directory, Flags))
+		if (!Request->Init(WatchDirectory, Flags))
 		{
-			UE_LOG(LogDirectoryWatcher, Warning, TEXT("Failed to begin reading directory changes for %s."), *Directory);
+			UE_LOG(LogDirectoryWatcher, Warning, TEXT("Failed to begin reading directory changes for %s."), *WatchDirectory);
 			delete Request;
 			NumRequests--;
 			return false;
 		}
 
-		RequestMap.Add(Directory, Request);
+		RequestMap.Add(WatchDirectory, Request);
 	}
 
 	OutHandle = Request->AddDelegate(InDelegate, Flags);
@@ -73,9 +75,11 @@ bool FDirectoryWatcherLinux::RegisterDirectoryChangedCallback_Handle(const FStri
 }
 
 
-bool FDirectoryWatcherLinux::UnregisterDirectoryChangedCallback_Handle(const FString& Directory, FDelegateHandle InHandle)
+bool FDirectoryWatcherLinux::UnregisterDirectoryChangedCallback_Handle(const FString& InDirectory, FDelegateHandle InHandle)
 {
-	FDirectoryWatchRequestLinux** RequestPtr = RequestMap.Find(Directory);
+	// Make sure the path is absolute
+	FString WatchDirectory = FPaths::ConvertRelativePathToFull(InDirectory);
+	FDirectoryWatchRequestLinux** RequestPtr = RequestMap.Find(WatchDirectory);
 
 	if (RequestPtr)
 	{
@@ -89,7 +93,7 @@ bool FDirectoryWatcherLinux::UnregisterDirectoryChangedCallback_Handle(const FSt
 			if (!Request->HasDelegates())
 			{
 				// Remove from the active map and add to the pending delete list
-				RequestMap.Remove(Directory);
+				RequestMap.Remove(WatchDirectory);
 				RequestsPendingDelete.AddUnique(Request);
 
 				// Signal to end the watch which will mark this request for deletion
@@ -115,9 +119,12 @@ void FDirectoryWatcherLinux::Tick(float DeltaSeconds)
 		RequestsPendingDelete.RemoveAt(RequestIdx);
 	}
 
-	// Trigger any file change notification delegates
-	for (TMap<FString, FDirectoryWatchRequestLinux*>::TConstIterator RequestIt(RequestMap); RequestIt; ++RequestIt)
-	{
-		RequestIt.Value()->ProcessPendingNotifications();
-	}
+	FDirectoryWatchRequestLinux::ProcessNotifications(RequestMap);
 }
+
+bool FDirectoryWatcherLinux::DumpStats()
+{
+	FDirectoryWatchRequestLinux::DumpStats(RequestMap);
+	return true;
+}
+
diff --git a/Engine/Source/Developer/DirectoryWatcher/Private/Linux/DirectoryWatcherLinux.h b/Engine/Source/Developer/DirectoryWatcher/Private/Linux/DirectoryWatcherLinux.h
--- a/Engine/Source/Developer/DirectoryWatcher/Private/Linux/DirectoryWatcherLinux.h
+++ b/Engine/Source/Developer/DirectoryWatcher/Private/Linux/DirectoryWatcherLinux.h
@@ -16,6 +16,7 @@ public:
 	virtual bool RegisterDirectoryChangedCallback_Handle (const FString& Directory, const FDirectoryChanged& InDelegate, FDelegateHandle& OutHandle, uint32 Flags) override;
 	virtual bool UnregisterDirectoryChangedCallback_Handle(const FString& Directory, FDelegateHandle InHandle) override;
 	virtual void Tick (float DeltaSeconds) override;
+	virtual bool DumpStats() override;
 
 	/** Map of directory paths to requests */
 	TMap<FString, FDirectoryWatchRequestLinux*> RequestMap;
diff --git a/Engine/Source/Developer/DirectoryWatcher/Public/IDirectoryWatcher.h b/Engine/Source/Developer/DirectoryWatcher/Public/IDirectoryWatcher.h
--- a/Engine/Source/Developer/DirectoryWatcher/Public/IDirectoryWatcher.h
+++ b/Engine/Source/Developer/DirectoryWatcher/Public/IDirectoryWatcher.h
@@ -72,5 +72,13 @@ public:
 	virtual void Tick(float DeltaSeconds)
 	{
 	}
+
+	/**
+	 * Allows for subclasses to dump notify statistics (returns true if implemented)
+	 */
+	virtual bool DumpStats()
+	{
+		return false;
+	}
 };
 
diff --git a/Engine/Source/Programs/TestPAL/Private/TestDirectoryWatcher.cpp b/Engine/Source/Programs/TestPAL/Private/TestDirectoryWatcher.cpp
--- a/Engine/Source/Programs/TestPAL/Private/TestDirectoryWatcher.cpp
+++ b/Engine/Source/Programs/TestPAL/Private/TestDirectoryWatcher.cpp
@@ -47,7 +47,7 @@ int32 DirectoryWatcherTest(const TCHAR* CommandLine)
 	UE_LOG(LogTestPAL, Display, TEXT("Running directory watcher test."));
 
 	IPlatformFile& PlatformFile = FPlatformFileManager::Get().GetPlatformFile();
-	FString TestDir = FString::Printf(TEXT("%sDirectoryWatcherTest%d"), FPlatformProcess::UserTempDir(), FPlatformProcess::GetCurrentProcessId());
+	FString TestDir = FString::Printf(TEXT("%s/DirectoryWatcherTest%d"), FPlatformProcess::UserTempDir(), FPlatformProcess::GetCurrentProcessId());
 
 	if (PlatformFile.CreateDirectory(*TestDir) && PlatformFile.CreateDirectory(*(TestDir + TEXT("/subtest"))))
 	{
@@ -57,8 +57,14 @@ int32 DirectoryWatcherTest(const TCHAR* CommandLine)
 		IDirectoryWatcher* DirectoryWatcher = FModuleManager::Get().LoadModuleChecked<FDirectoryWatcherModule>(TEXT("DirectoryWatcher")).Get();
 		if (DirectoryWatcher)
 		{
+			/** Whether to include notifications for changes to actual directories (such as directories being created or removed). */
+			//   IDirectoryWatcher::WatchOptions::IgnoreChangesInSubtree
+			/** Whether changes in subdirectories need to be reported. */
+			//   IDirectoryWatcher::WatchOptions::IncludeDirectoryChanges
+			uint32 Flags = 0;
+
 			auto Callback = IDirectoryWatcher::FDirectoryChanged::CreateRaw(&Detector, &FChangeDetector::OnDirectoryChanged);
-			DirectoryWatcher->RegisterDirectoryChangedCallback_Handle(TestDir, Callback, DirectoryChangedHandle);
+			DirectoryWatcher->RegisterDirectoryChangedCallback_Handle(TestDir, Callback, DirectoryChangedHandle, Flags);
 			UE_LOG(LogTestPAL, Display, TEXT("Registered callback for changes in '%s'"), *TestDir);
 		}
 		else
@@ -180,6 +186,7 @@ int32 DirectoryWatcherTest(const TCHAR* CommandLine)
 			DirectoryWatcher->Tick(1.0f);
 		}
 
+		DirectoryWatcher->DumpStats();
 
 		// clean up
 		verify(DirectoryWatcher->UnregisterDirectoryChangedCallback_Handle(TestDir, DirectoryChangedHandle));
