Add subdirectory changes to list of changes in Linux directory watcher

Previously, we only added the new directory and not the contents. The new AssetDataGatherer has this assumption now though:

  // Note that we AddDiscovered but do not scan the directory
  // Any files and paths under it will be added by their own event from the directory watcher, so a scan is unnecessary.
  // The directory may also be scanned in the future because a parent directory is still yet pending to scan,
  // we do not try to prevent that wasteful rescan because this is a rare event and it does not cause a behavior problem

When we iterate the contents of a new directory we now add them to the FileChanges array to match the above.

diff --git a/Engine/Source/Developer/DirectoryWatcher/Private/Linux/DirectoryWatchRequestLinux.cpp b/Engine/Source/Developer/DirectoryWatcher/Private/Linux/DirectoryWatchRequestLinux.cpp
--- a/Engine/Source/Developer/DirectoryWatcher/Private/Linux/DirectoryWatchRequestLinux.cpp
+++ b/Engine/Source/Developer/DirectoryWatcher/Private/Linux/DirectoryWatchRequestLinux.cpp
@@ -85,8 +85,8 @@ bool FDirectoryWatchRequestLinux::Init(const FString& InDirectory, uint32 Flags)
 		return false;
 	}
 
-	// find all subdirs
-	WatchDirectoryTree(WatchDirectory);
+	// Find all subdirs and add inotify watch requests
+	WatchDirectoryTree(WatchDirectory, nullptr);
 
 	return true;
 }
@@ -174,7 +174,7 @@ void FDirectoryWatchRequestLinux::ProcessPendingNotifications()
 	}
 }
 
-void FDirectoryWatchRequestLinux::WatchDirectoryTree(const FString & RootAbsolutePath)
+void FDirectoryWatchRequestLinux::WatchDirectoryTree(const FString & RootAbsolutePath, TArray<TPair<FFileChangeData, bool>>* FileChanges)
 {
 	if (bEndWatchRequestInvoked || (FileDescriptor == -1))
 	{
@@ -190,16 +190,26 @@ void FDirectoryWatchRequestLinux::WatchDirectoryTree(const FString & RootAbsolut
 
 	UE_LOG(LogDirectoryWatcher, VeryVerbose, TEXT("Watching tree '%s'"), *RootAbsolutePath);
 
+	if (FileChanges)
+	{
+		FileChanges->Emplace(FFileChangeData(RootAbsolutePath, FFileChangeData::FCA_Added), true);
+	}
+
 	TArray<FString> AllFiles;
 	if (bWatchSubtree)
 	{
 		IPlatformFile::GetPlatformPhysical().IterateDirectoryRecursively(*RootAbsolutePath,
-			[&AllFiles](const TCHAR* Name, bool bIsDirectory)
+			[&AllFiles, FileChanges](const TCHAR* Name, bool bIsDirectory)
 				{
 					if (bIsDirectory)
 					{
 						AllFiles.Add(Name);
 					}
+
+					if (FileChanges)
+					{
+						FileChanges->Emplace(FFileChangeData(Name, FFileChangeData::FCA_Added), bIsDirectory);
+					}
 					return true;
 				});
 	}
@@ -374,11 +384,14 @@ void FDirectoryWatchRequestLinux::ProcessChanges(TArray<TPair<FFileChangeData, b
 						// IN_MOVED_TO: Generated for the directory containing the new filename when a file is renamed
 						if (bIsDir)
 						{
-							// If a directory was created/moved, watch it
-							WatchDirectoryTree(AffectedFile);
+							// If a directory was created/moved, watch it and add changes to FileChanges.
+							// Leave Action as FCA_Unknown so nothing gets added down below.
+							WatchDirectoryTree(AffectedFile, &FileChanges);
+						}
+						else
+						{
+							Action = FFileChangeData::FCA_Added;
 						}
-
-						Action = FFileChangeData::FCA_Added;
 					}
 					else if (Event->mask & IN_MODIFY)
 					{
@@ -432,7 +445,7 @@ void FDirectoryWatchRequestLinux::ProcessChanges(TArray<TPair<FFileChangeData, b
 						Action = FFileChangeData::FCA_Removed;
 					}
 
-					if (Event->len)
+					if (Event->len && (Action != FFileChangeData::FCA_Unknown))
 					{
 						FileChanges.Emplace(FFileChangeData(AffectedFile, Action), bIsDir);
 					}
diff --git a/Engine/Source/Developer/DirectoryWatcher/Private/Linux/DirectoryWatchRequestLinux.h b/Engine/Source/Developer/DirectoryWatcher/Private/Linux/DirectoryWatchRequestLinux.h
--- a/Engine/Source/Developer/DirectoryWatcher/Private/Linux/DirectoryWatchRequestLinux.h
+++ b/Engine/Source/Developer/DirectoryWatcher/Private/Linux/DirectoryWatchRequestLinux.h
@@ -37,7 +37,7 @@ private:
 	void ProcessPendingNotifications();
 
 	/** Adds watches for all files (and subdirectories) in a directory. */
-	void WatchDirectoryTree(const FString& RootAbsolutePath);
+	void WatchDirectoryTree(const FString& RootAbsolutePath, TArray<TPair<FFileChangeData, bool>>* FileChanges);
 
 	/** Removes all watches for path */
 	void UnwatchDirectoryTree(const FString& RootAbsolutePath);
