Split logic as previous change may likely introduce some bad issues

diff --git a/Engine/Source/Developer/DirectoryWatcher/Private/Linux/DirectoryWatchRequestLinux.cpp b/Engine/Source/Developer/DirectoryWatcher/Private/Linux/DirectoryWatchRequestLinux.cpp
--- a/Engine/Source/Developer/DirectoryWatcher/Private/Linux/DirectoryWatchRequestLinux.cpp
+++ b/Engine/Source/Developer/DirectoryWatcher/Private/Linux/DirectoryWatchRequestLinux.cpp
@@ -203,14 +203,18 @@ void FDirectoryWatchRequestLinux::WatchDirectoryTree(const FString & RootAbsolut
 		return;
 	}
 
+	// If this isn't our root watch directory or under it, don't watch
+	if (!RootAbsolutePath.StartsWith(WatchDirectory, ESearchCase::CaseSensitive))
+	{
+		return;
+	}
+
 	if (FileChangesPtr)
 	{
 		FileChangesPtr->Emplace(FFileChangeData(RootAbsolutePath, FFileChangeData::FCA_Added), true);
 	}
 
-	// If this isn't our root watch directory or under it, don't watch
-	if (!RootAbsolutePath.StartsWith(WatchDirectory, ESearchCase::CaseSensitive) ||
-		(!bWatchSubtree && (RootAbsolutePath != WatchDirectory)))
+	if (!bWatchSubtree && (RootAbsolutePath != WatchDirectory))
 	{
 		return;
 	}
