Fix unbound delegate crash in directory watcher

diff --git a/Engine/Source/Developer/DirectoryWatcher/Private/Windows/DirectoryWatchRequestWindows.cpp b/Engine/Source/Developer/DirectoryWatcher/Private/Windows/DirectoryWatchRequestWindows.cpp
--- a/Engine/Source/Developer/DirectoryWatcher/Private/Windows/DirectoryWatchRequestWindows.cpp
+++ b/Engine/Source/Developer/DirectoryWatcher/Private/Windows/DirectoryWatchRequestWindows.cpp
@@ -163,7 +163,14 @@ void FDirectoryWatchRequestWindows::ProcessPendingNotifications()
 	{
 		for (int32 DelegateIdx = 0; DelegateIdx < Delegates.Num(); ++DelegateIdx)
 		{
-			Delegates[DelegateIdx].Execute(FileChanges);
+			if (Delegates[DelegateIdx].IsBound())
+			{
+				Delegates[DelegateIdx].Execute(FileChanges);
+			}
+			else
+			{
+				Delegates.RemoveAt(DelegateIdx--);
+			}
 		}
 
 		FileChanges.Empty();
