DirectoryWatcherWindows: Handle ERROR_INVALID_PARAMETER when buffer size is too large to watch a network file share.

diff --git a/Engine/Source/Developer/DirectoryWatcher/Private/Windows/DirectoryWatchRequestWindows.cpp b/Engine/Source/Developer/DirectoryWatcher/Private/Windows/DirectoryWatchRequestWindows.cpp
--- a/Engine/Source/Developer/DirectoryWatcher/Private/Windows/DirectoryWatchRequestWindows.cpp
+++ b/Engine/Source/Developer/DirectoryWatcher/Private/Windows/DirectoryWatchRequestWindows.cpp
@@ -8,7 +8,6 @@ FDirectoryWatchRequestWindows::FDirectoryWatchRequestWindows(uint32 Flags)
 	bPendingDelete = false;
 	bEndWatchRequestInvoked = false;
 
-	MaxChanges = 16384;
 	bWatchSubtree = (Flags & IDirectoryWatcher::WatchOptions::IgnoreChangesInSubtree) == 0;
 	bool bIncludeDirectoryEvents = (Flags & IDirectoryWatcher::WatchOptions::IncludeDirectoryChanges) != 0;
 
@@ -16,36 +15,40 @@ FDirectoryWatchRequestWindows::FDirectoryWatchRequestWindows(uint32 Flags)
 
 	DirectoryHandle = INVALID_HANDLE_VALUE;
 
-	BufferLength = sizeof(FILE_NOTIFY_INFORMATION) * MaxChanges;
-	Buffer = FMemory::Malloc(BufferLength, alignof(DWORD));
-	BackBuffer = FMemory::Malloc(BufferLength, alignof(DWORD));
+	constexpr int32 InitialMaxChanges = 16384;
+	SetBufferByChangeCount(InitialMaxChanges);
 
 	FMemory::Memzero(&Overlapped, sizeof(Overlapped));
-	FMemory::Memzero(Buffer, BufferLength);
-	FMemory::Memzero(BackBuffer, BufferLength);
 
 	Overlapped.hEvent = this;
 }
 
 FDirectoryWatchRequestWindows::~FDirectoryWatchRequestWindows()
 {
-	if (Buffer)
-	{
-		FMemory::Free(Buffer);
-	}
-
-	if (BackBuffer)
-	{
-		FMemory::Free(BackBuffer);
-	}
-
 	if ( DirectoryHandle != INVALID_HANDLE_VALUE )
 	{
 		::CloseHandle(DirectoryHandle);
 		DirectoryHandle = INVALID_HANDLE_VALUE;
+		bBufferInUse = false;
 	}
 }
 
+void FDirectoryWatchRequestWindows::SetBufferByChangeCount(int32 MaxChanges)
+{
+	checkf(!bBufferInUse, TEXT("Reallocating the buffer while it is referenced in a pending ReadDirectoryChangesW call is invalid and will likely cause a crash."));
+	BufferLength = sizeof(FILE_NOTIFY_INFORMATION) * MaxChanges;
+	Buffer.Reset(reinterpret_cast<uint8*>(FMemory::Malloc(BufferLength, alignof(DWORD))));
+	BackBuffer.Reset(reinterpret_cast<uint8*>(FMemory::Malloc(BufferLength, alignof(DWORD))));
+	FMemory::Memzero(Buffer.Get(), BufferLength);
+	FMemory::Memzero(BackBuffer.Get(), BufferLength);
+}
+
+void FDirectoryWatchRequestWindows::SetBufferBySize(int32 BufferSize)
+{
+	int32 MaxChanges = BufferSize / sizeof(FILE_NOTIFY_INFORMATION); // Arbitrarily we choose to round down
+	SetBufferByChangeCount(MaxChanges);
+}
+
 bool FDirectoryWatchRequestWindows::Init(const FString& InDirectory)
 {
 	check(Buffer);
@@ -82,13 +85,13 @@ bool FDirectoryWatchRequestWindows::Init(const FString& InDirectory)
 	{
 		const DWORD ErrorCode = ::GetLastError();
 		// Failed to obtain a handle to this directory
-		UE_LOG(LogDirectoryWatcher, Display, TEXT("Failed initialise directory notification handle for '%s', and we were unable to create a new request. GetLastError code [%d]"), *FullPath, ErrorCode);
+		UE_LOG(LogDirectoryWatcher, Display, TEXT("CreateFile failed for '%s'. GetLastError code [%d]"), *FullPath, ErrorCode);
 		return false;
 	}
 
 	const bool bSuccess = !!::ReadDirectoryChangesW(
 		DirectoryHandle,
-		Buffer,
+		Buffer.Get(),
 		BufferLength,
 		bWatchSubtree,
 		NotifyFilter,
@@ -99,11 +102,12 @@ bool FDirectoryWatchRequestWindows::Init(const FString& InDirectory)
 	if ( !bSuccess  )
 	{
 		const DWORD ErrorCode = ::GetLastError();
-		UE_LOG(LogDirectoryWatcher, Display, TEXT("An initialise directory notification failed for '%s', and we were unable to create a new request. GetLastError code [%d]"), *Directory, ErrorCode);
+		UE_LOG(LogDirectoryWatcher, Display, TEXT("Initial ReadDirectoryChangesW failed for '%s'. GetLastError code [%d]"), *Directory, ErrorCode);
 		::CloseHandle(DirectoryHandle);
 		DirectoryHandle = INVALID_HANDLE_VALUE;
 		return false;
 	}
+	bBufferInUse = true;
 
 	return true;
 }
@@ -138,6 +142,7 @@ void FDirectoryWatchRequestWindows::EndWatchRequest()
 		if ( DirectoryHandle != INVALID_HANDLE_VALUE )
 		{
 			CancelIoEx(DirectoryHandle, &Overlapped);
+			bBufferInUse = false;
 			// Clear the handle so we don't setup any more requests, and wait for the operation to finish
 			HANDLE TempDirectoryHandle = DirectoryHandle;
 			DirectoryHandle = INVALID_HANDLE_VALUE;
@@ -179,6 +184,14 @@ void FDirectoryWatchRequestWindows::ProcessPendingNotifications()
 
 void FDirectoryWatchRequestWindows::ProcessChange(uint32 Error, uint32 NumBytes)
 {
+	bBufferInUse = false; // Buffer reallocations are allowed in this handling code before we resubmit the watch
+	auto CloseHandleAndMarkForDelete = [this]()
+	{
+		::CloseHandle(DirectoryHandle);
+		DirectoryHandle = INVALID_HANDLE_VALUE;
+		bPendingDelete = true;
+	};
+
 	if (Error == ERROR_OPERATION_ABORTED) 
 	{
 		// The operation was aborted, likely due to EndWatchRequest canceling it.
@@ -187,42 +200,47 @@ void FDirectoryWatchRequestWindows::ProcessChange(uint32 Error, uint32 NumBytes)
 		UE_CLOG(!IsEngineExitRequested(), LogDirectoryWatcher, Log, TEXT("A directory notification for '%s' was aborted."), *Directory);
 		return; 
 	}
-
-	const bool bValidNotification = (Error != ERROR_IO_INCOMPLETE && NumBytes > 0 );
-	const bool bAccessError = (Error == ERROR_ACCESS_DENIED);
-
-	auto CloseHandleAndMarkForDelete = [this]()
-	{
-		::CloseHandle(DirectoryHandle);
-		DirectoryHandle = INVALID_HANDLE_VALUE;
-		bPendingDelete = true;
-	};
-
-	// Swap the pointer to the backbuffer so we can start a new read as soon as possible
-	if ( bValidNotification )
+	const bool bValidNotification = Error != ERROR_IO_INCOMPLETE && NumBytes > 0;
+	if (bValidNotification)
 	{
+		// Swap the pointer to the backbuffer so we can start a new read as soon as possible
 		Swap(Buffer, BackBuffer);
 		check(Buffer && BackBuffer);
 	}
-
-	if ( !bValidNotification )
+	else if (Error == ERROR_INVALID_PARAMETER)
 	{
-		if (bAccessError)
+		// ReadDirectoryChangesW fails with ERROR_INVALID_PARAMETER when the buffer length is greater than 64 KB and the application is
+		// monitoring a directory over the network. This is due to a packet size limitation with the underlying file sharing protocols.
+		constexpr int32 MaxAllowedSize = 64 * 1024;
+		if (BufferLength > MaxAllowedSize)
 		{
-			CloseHandleAndMarkForDelete();
-			UE_LOG(LogDirectoryWatcher, Display, TEXT("A directory notification failed for '%s' because it could not be accessed. Aborting watch request..."), *Directory);
-			return;
+			// This error is expected, and is sent before we fail to read changes, so do not log it
+			SetBufferBySize(MaxAllowedSize);
 		}
 		else
 		{
-			UE_LOG(LogDirectoryWatcher, Display, TEXT("A directory notification failed for '%s' because it was empty or there was a buffer overflow. Attemping another request..."), *Directory);
+			UE_LOG(LogDirectoryWatcher, Display, TEXT("A directory notification failed for '%s' with ERROR_INVALID_PARAMETER. Attempting another request..."), *Directory);
 		}
 	}
+	else if (Error == ERROR_ACCESS_DENIED)
+	{
+		CloseHandleAndMarkForDelete();
+		UE_LOG(LogDirectoryWatcher, Display, TEXT("A directory notification failed for '%s' because it could not be accessed. Aborting watch request..."), *Directory);
+		return;
+	}
+	else if (Error != ERROR_SUCCESS)
+	{
+		UE_LOG(LogDirectoryWatcher, Display, TEXT("A directory notification failed for '%s' with error code [%d]. Attemping another request..."), *Directory, Error);
+	}
+	else
+	{
+		UE_LOG(LogDirectoryWatcher, Display, TEXT("A directory notification failed for '%s' due to buffer overflow. Attemping another request..."), *Directory);
+	}
 
 	// Start up another read
 	const bool bSuccess = !!::ReadDirectoryChangesW(
 		DirectoryHandle,
-		Buffer,
+		Buffer.Get(),
 		BufferLength,
 		bWatchSubtree,
 		NotifyFilter,
@@ -233,11 +251,12 @@ void FDirectoryWatchRequestWindows::ProcessChange(uint32 Error, uint32 NumBytes)
 	if ( !bSuccess  )
 	{
 		const DWORD ErrorCode = ::GetLastError();
-		UE_LOG(LogDirectoryWatcher, Display, TEXT("A directory notification failed, and we were unable to create a new request. GetLastError code [%d] Handle [%p], Path [%s] "), ErrorCode, DirectoryHandle, *Directory);
-		// Failed to re-create the read request.
+		UE_LOG(LogDirectoryWatcher, Display, TEXT("Refresh of ReadDirectoryChangesW failed. GetLastError code [%d] Handle [%p], Path [%s]. Aborting watch request..."),
+			ErrorCode, DirectoryHandle, *Directory);
 		CloseHandleAndMarkForDelete();
 		return;
 	}
+	bBufferInUse = true;
 
 	// No need to process the change if we can not execute any delegates
 	if ( !HasDelegates() || !bValidNotification )
@@ -246,7 +265,7 @@ void FDirectoryWatchRequestWindows::ProcessChange(uint32 Error, uint32 NumBytes)
 	}
 
 	// Process the change
-	void* InfoBase = BackBuffer;
+	uint8* InfoBase = BackBuffer.Get();
 	do
 	{
 		FILE_NOTIFY_INFORMATION* NotifyInfo = (FILE_NOTIFY_INFORMATION*)InfoBase;
@@ -292,7 +311,7 @@ void FDirectoryWatchRequestWindows::ProcessChange(uint32 Error, uint32 NumBytes)
 		}
 
 		// Adjust the offset and update the NotifyInfo pointer
-		InfoBase = (uint8*)InfoBase + NotifyInfo->NextEntryOffset;
+		InfoBase = InfoBase + NotifyInfo->NextEntryOffset;
 	}
 	while(true);
 }
diff --git a/Engine/Source/Developer/DirectoryWatcher/Private/Windows/DirectoryWatchRequestWindows.h b/Engine/Source/Developer/DirectoryWatcher/Private/Windows/DirectoryWatchRequestWindows.h
--- a/Engine/Source/Developer/DirectoryWatcher/Private/Windows/DirectoryWatchRequestWindows.h
+++ b/Engine/Source/Developer/DirectoryWatcher/Private/Windows/DirectoryWatchRequestWindows.h
@@ -2,9 +2,11 @@
 
 #pragma once
 
-#include "Windows/WindowsHWrapper.h"
 #include "Delegates/Delegate.h"
 #include "IDirectoryWatcher.h"
+#include "HAL/UnrealMemory.h"
+#include "Templates/UniquePtr.h"
+#include "Windows/WindowsHWrapper.h"
 
 class FDirectoryWatchRequestWindows
 {
@@ -33,21 +35,31 @@ public:
 private:
 	/** Non-static handler for an OS notification of a directory change */
 	void ProcessChange(uint32 Error, uint32 NumBytes);
+	void SetBufferByChangeCount(int32 MaxChanges);
+	void SetBufferBySize(int32 BufferSize);
+
 	/** Static Handler for an OS notification of a directory change */
 	static void CALLBACK ChangeNotification(::DWORD Error, ::DWORD NumBytes, LPOVERLAPPED InOverlapped);
 
 	FString Directory;
 	HANDLE DirectoryHandle;
-	int32 MaxChanges;
 	bool bWatchSubtree;
 	uint32 NotifyFilter;
-	uint32 BufferLength;
-	void* Buffer;
-	void* BackBuffer;
+	uint32 BufferLength = 0;
+	struct FDeleterFree
+	{
+		void operator()(uint8* Ptr) const
+		{
+			if (Ptr) { FMemory::Free(Ptr); }
+		}
+	};
+	TUniquePtr<uint8, FDeleterFree> Buffer;
+	TUniquePtr<uint8, FDeleterFree> BackBuffer;
 	OVERLAPPED Overlapped;
 
 	bool bPendingDelete;
 	bool bEndWatchRequestInvoked;
+	bool bBufferInUse = false;
 
 	TArray<IDirectoryWatcher::FDirectoryChanged> Delegates;
 	TArray<FFileChangeData> FileChanges;
