AssetRegistry,DirectoryWatcher: Add new event type to DirectoryWatcher, to handle notifications on Windows that the directory needs to be rescanned.
Handle this event type in AssetRegistry and rescan.
#rnx

diff --git a/Engine/Source/Developer/DirectoryWatcher/Private/Windows/DirectoryWatchRequestWindows.cpp b/Engine/Source/Developer/DirectoryWatcher/Private/Windows/DirectoryWatchRequestWindows.cpp
--- a/Engine/Source/Developer/DirectoryWatcher/Private/Windows/DirectoryWatchRequestWindows.cpp
+++ b/Engine/Source/Developer/DirectoryWatcher/Private/Windows/DirectoryWatchRequestWindows.cpp
@@ -1,12 +1,15 @@
 // Copyright Epic Games, Inc. All Rights Reserved.
 
 #include "DirectoryWatchRequestWindows.h"
+
 #include "DirectoryWatcherPrivate.h"
+#include "Misc/DateTime.h"
 
 FDirectoryWatchRequestWindows::FDirectoryWatchRequestWindows(uint32 Flags)
 {
 	bPendingDelete = false;
 	bEndWatchRequestInvoked = false;
+	WatchStartedTimestamp = 0;
 
 	bWatchSubtree = (Flags & IDirectoryWatcher::WatchOptions::IgnoreChangesInSubtree) == 0;
 	bool bIncludeDirectoryEvents = (Flags & IDirectoryWatcher::WatchOptions::IncludeDirectoryChanges) != 0;
@@ -107,6 +110,7 @@ bool FDirectoryWatchRequestWindows::Init(const FString& InDirectory)
 		DirectoryHandle = INVALID_HANDLE_VALUE;
 		return false;
 	}
+	WatchStartedTimestamp = FDateTime::UtcNow().ToUnixTimestamp();
 	bBufferInUse = true;
 
 	return true;
@@ -200,7 +204,9 @@ void FDirectoryWatchRequestWindows::ProcessChange(uint32 Error, uint32 NumBytes)
 		UE_CLOG(!IsEngineExitRequested(), LogDirectoryWatcher, Log, TEXT("A directory notification for '%s' was aborted."), *Directory);
 		return; 
 	}
-	const bool bValidNotification = Error != ERROR_IO_INCOMPLETE && NumBytes > 0;
+	bool bValidNotification = Error != ERROR_IO_INCOMPLETE && NumBytes > 0;
+	bool bIsRescan = false;
+	int64 PreviousWatchStartedTimeStamp = 0;
 	if (bValidNotification)
 	{
 		// Swap the pointer to the backbuffer so we can start a new read as soon as possible
@@ -228,6 +234,12 @@ void FDirectoryWatchRequestWindows::ProcessChange(uint32 Error, uint32 NumBytes)
 		UE_LOG(LogDirectoryWatcher, Display, TEXT("A directory notification failed for '%s' because it could not be accessed. Aborting watch request..."), *Directory);
 		return;
 	}
+	else if (Error == ERROR_NOTIFY_ENUM_DIR)
+	{
+		bValidNotification = true;
+		bIsRescan = true;
+		PreviousWatchStartedTimeStamp = WatchStartedTimestamp;
+	}
 	else if (Error != ERROR_SUCCESS)
 	{
 		UE_LOG(LogDirectoryWatcher, Display, TEXT("A directory notification failed for '%s' with error code [%d]. Attemping another request..."), *Directory, Error);
@@ -256,6 +268,7 @@ void FDirectoryWatchRequestWindows::ProcessChange(uint32 Error, uint32 NumBytes)
 		CloseHandleAndMarkForDelete();
 		return;
 	}
+	WatchStartedTimestamp = FDateTime::UtcNow().ToUnixTimestamp();
 	bBufferInUse = true;
 
 	// No need to process the change if we can not execute any delegates
@@ -264,56 +277,64 @@ void FDirectoryWatchRequestWindows::ProcessChange(uint32 Error, uint32 NumBytes)
 		return;
 	}
 
-	// Process the change
-	uint8* InfoBase = BackBuffer.Get();
-	do
+	if (!bIsRescan)
 	{
-		FILE_NOTIFY_INFORMATION* NotifyInfo = (FILE_NOTIFY_INFORMATION*)InfoBase;
-
-		// Copy the WCHAR out of the NotifyInfo so we can put a NULL terminator on it and convert it to a FString
-		FString LeafFilename;
+		// Process the change
+		uint8* InfoBase = BackBuffer.Get();
+		do
 		{
-			// The Memcpy below assumes that WCHAR and TCHAR are equivalent (which they should be on Windows)
-			static_assert(sizeof(WCHAR) == sizeof(TCHAR), "WCHAR is assumed to be the same size as TCHAR on Windows!");
-
-			const int32 LeafFilenameLen = NotifyInfo->FileNameLength / sizeof(WCHAR);
-			LeafFilename.GetCharArray().AddZeroed(LeafFilenameLen + 1);
-			FMemory::Memcpy(LeafFilename.GetCharArray().GetData(), NotifyInfo->FileName, NotifyInfo->FileNameLength);
-		}
+			FILE_NOTIFY_INFORMATION* NotifyInfo = (FILE_NOTIFY_INFORMATION*)InfoBase;
 
-		FFileChangeData::EFileChangeAction Action;
-		switch(NotifyInfo->Action)
-		{
-			case FILE_ACTION_ADDED:
-			case FILE_ACTION_RENAMED_NEW_NAME:
-				Action = FFileChangeData::FCA_Added;
-				break;
+			// Copy the WCHAR out of the NotifyInfo so we can put a NULL terminator on it and convert it to a FString
+			FString LeafFilename;
+			{
+				// The Memcpy below assumes that WCHAR and TCHAR are equivalent (which they should be on Windows)
+				static_assert(sizeof(WCHAR) == sizeof(TCHAR), "WCHAR is assumed to be the same size as TCHAR on Windows!");
 
-			case FILE_ACTION_REMOVED:
-			case FILE_ACTION_RENAMED_OLD_NAME:
-				Action = FFileChangeData::FCA_Removed;
-				break;
+				const int32 LeafFilenameLen = NotifyInfo->FileNameLength / sizeof(WCHAR);
+				LeafFilename.GetCharArray().AddZeroed(LeafFilenameLen + 1);
+				FMemory::Memcpy(LeafFilename.GetCharArray().GetData(), NotifyInfo->FileName, NotifyInfo->FileNameLength);
+			}
 
-			case FILE_ACTION_MODIFIED:
-				Action = FFileChangeData::FCA_Modified;
-				break;
+			FFileChangeData::EFileChangeAction Action;
+			switch(NotifyInfo->Action)
+			{
+				case FILE_ACTION_ADDED:
+				case FILE_ACTION_RENAMED_NEW_NAME:
+					Action = FFileChangeData::FCA_Added;
+					break;
+
+				case FILE_ACTION_REMOVED:
+				case FILE_ACTION_RENAMED_OLD_NAME:
+					Action = FFileChangeData::FCA_Removed;
+					break;
+
+				case FILE_ACTION_MODIFIED:
+					Action = FFileChangeData::FCA_Modified;
+					break;
+
+				default:
+					Action = FFileChangeData::FCA_Unknown;
+					break;
+			}
+			FileChanges.Emplace(Directory / LeafFilename, Action);
 
-			default:
-				Action = FFileChangeData::FCA_Unknown;
+			// If there is not another entry, break the loop
+			if ( NotifyInfo->NextEntryOffset == 0 )
+			{
 				break;
-		}
-		FileChanges.Emplace(Directory / LeafFilename, Action);
+			}
 
-		// If there is not another entry, break the loop
-		if ( NotifyInfo->NextEntryOffset == 0 )
-		{
-			break;
+			// Adjust the offset and update the NotifyInfo pointer
+			InfoBase = InfoBase + NotifyInfo->NextEntryOffset;
 		}
-
-		// Adjust the offset and update the NotifyInfo pointer
-		InfoBase = InfoBase + NotifyInfo->NextEntryOffset;
+		while(true);
+	}
+	else
+	{
+		FFileChangeData& ChangeData = FileChanges.Emplace_GetRef(Directory, FFileChangeData::FCA_RescanRequired);
+		ChangeData.TimeStamp = PreviousWatchStartedTimeStamp;
 	}
-	while(true);
 }
 
 void FDirectoryWatchRequestWindows::ChangeNotification(::DWORD Error, ::DWORD NumBytes, LPOVERLAPPED InOverlapped)
diff --git a/Engine/Source/Developer/DirectoryWatcher/Private/Windows/DirectoryWatchRequestWindows.h b/Engine/Source/Developer/DirectoryWatcher/Private/Windows/DirectoryWatchRequestWindows.h
--- a/Engine/Source/Developer/DirectoryWatcher/Private/Windows/DirectoryWatchRequestWindows.h
+++ b/Engine/Source/Developer/DirectoryWatcher/Private/Windows/DirectoryWatchRequestWindows.h
@@ -43,7 +43,7 @@ private:
 
 	FString Directory;
 	HANDLE DirectoryHandle;
-	bool bWatchSubtree;
+	int64 WatchStartedTimestamp;
 	uint32 NotifyFilter;
 	uint32 BufferLength = 0;
 	struct FDeleterFree
@@ -60,6 +60,7 @@ private:
 	bool bPendingDelete;
 	bool bEndWatchRequestInvoked;
 	bool bBufferInUse = false;
+	bool bWatchSubtree;
 
 	TArray<IDirectoryWatcher::FDirectoryChanged> Delegates;
 	TArray<FFileChangeData> FileChanges;
diff --git a/Engine/Source/Developer/DirectoryWatcher/Public/IDirectoryWatcher.h b/Engine/Source/Developer/DirectoryWatcher/Public/IDirectoryWatcher.h
--- a/Engine/Source/Developer/DirectoryWatcher/Public/IDirectoryWatcher.h
+++ b/Engine/Source/Developer/DirectoryWatcher/Public/IDirectoryWatcher.h
@@ -12,7 +12,8 @@ struct FFileChangeData
 		FCA_Unknown,
 		FCA_Added,
 		FCA_Modified,
-		FCA_Removed
+		FCA_Removed,
+		FCA_RescanRequired,
 	};
 
 	FFileChangeData(const FString& InFilename, EFileChangeAction InAction)
@@ -22,8 +23,20 @@ struct FFileChangeData
 		FPaths::MakeStandardFilename(Filename);
 	}
 
+	/**
+	 * If the Action references a specific file, the name of the file.
+	 * Applies To: FCA_Added, FCA_Modified, FCA_Removed, FCA_RescanRequired.
+	 * For all other actions value will be emptystring.
+	 */
 	FString Filename;
-	EFileChangeAction Action;
+	/**
+	 * If the Action references a timestamp, the timezone UTC UnixTimeStamp (e.g. FDateTime::ToUnixTimeStamp) of the Action.
+	 * Applies to: FCA_RescanRequired.
+	 * For all other actions value will be 0.
+	 */
+	int64 TimeStamp = 0;
+	/** The reported Action. */
+	EFileChangeAction Action = EFileChangeAction::FCA_Unknown;
 };
 
 
diff --git a/Engine/Source/Runtime/AssetRegistry/Private/AssetRegistry.cpp b/Engine/Source/Runtime/AssetRegistry/Private/AssetRegistry.cpp
--- a/Engine/Source/Runtime/AssetRegistry/Private/AssetRegistry.cpp
+++ b/Engine/Source/Runtime/AssetRegistry/Private/AssetRegistry.cpp
@@ -5251,6 +5251,12 @@ void FAssetRegistryImpl::OnDirectoryChanged(Impl::FEventContext& EventContext,
 	TArray<FString> ModifiedFiles;
 	for (int32 FileIdx = 0; FileIdx < FileChangesProcessed.Num(); ++FileIdx)
 	{
+		if (FileChangesProcessed[FileIdx].Action == FFileChangeData::FCA_RescanRequired)
+		{
+			OnDirectoryRescanRequired(EventContext, InheritanceContext, FileChangesProcessed[FileIdx].Filename,
+				FileChangesProcessed[FileIdx].TimeStamp);
+			continue;
+		}
 		FString LongPackageName;
 		const FString File = FString(FileChangesProcessed[FileIdx].Filename);
 		const bool bIsPackageFile = FPackageName::IsPackageExtension(*FPaths::GetExtension(File, true));
@@ -5389,6 +5395,100 @@ void FAssetRegistryImpl::OnDirectoryChanged(Impl::FEventContext& EventContext,
 	ScanModifiedAssetFiles(EventContext, InheritanceContext, ModifiedFiles);
 }
 
+void FAssetRegistryImpl::OnDirectoryRescanRequired(Impl::FEventContext& EventContext,
+	Impl::FClassInheritanceContext& InheritanceContext, FString& DirPath, int64 BeforeTimeStamp)
+{
+	FString PackageNamePath;
+	FString NormalizedDirPath = FPaths::CreateStandardFilename(DirPath);
+	if (!FPackageName::TryConvertFilenameToLongPackageName(NormalizedDirPath, PackageNamePath))
+	{
+		return;
+	}
+
+	TArray<FString> NewFiles;
+	TArray<FString> NewDirs;
+	TArray<FString> ModifiedFiles;
+	TSet<FName> RemovedLongPackageNames;
+
+	FDateTime BeforeDateTime = FDateTime::FromUnixTimestamp(BeforeTimeStamp);
+	EnumerateAssetsByPathNoTags(*PackageNamePath, [&RemovedLongPackageNames](const FAssetData& AssetData)
+		{
+			RemovedLongPackageNames.Add(AssetData.PackageName);
+			return true;
+		}, true /* bRecursive */, true /* bIncludeOnlyOnDiskAssets */);
+
+	FPackageName::IteratePackagesInDirectory(DirPath,
+		[&DirPath, &PackageNamePath, &BeforeDateTime, &NewFiles, &ModifiedFiles, &RemovedLongPackageNames, &NormalizedDirPath]
+		(const TCHAR* Filename, const FFileStatData& StatData)
+		{
+			// Convert Filename to a PackagePath. We know the base dir so its faster to use that than FPackageName
+			// which has to scan all mount dirs
+			FStringView RelPath;
+			FString NormalizedFilename = FPaths::CreateStandardFilename(Filename);
+			if (!FPathViews::TryMakeChildPathRelativeTo(NormalizedFilename, NormalizedDirPath, RelPath))
+			{
+				return true;
+			}
+			const bool bIsPackageFile = FPackageName::IsPackageExtension(*FString(FPathViews::GetExtension(RelPath, true /* bIncludeDot */)));
+			RelPath = FPathViews::GetBaseFilenameWithPath(RelPath);
+			TStringBuilder<256> FilePackagePath;
+			FilePackagePath << PackageNamePath;
+			FPathViews::AppendPath(FilePackagePath, RelPath);
+			for (int32 Index = 0; Index < FilePackagePath.Len(); ++Index)
+			{
+				TCHAR& Char = FilePackagePath.GetData()[Index];
+				if (Char == '\\')
+				{
+					Char = '/';
+				}
+			}
+			const bool bIsValidPackageName = FPackageName::IsValidTextForLongPackageName(FilePackagePath);
+			if (!bIsPackageFile || !bIsValidPackageName)
+			{
+				return true;
+			}
+
+			if (StatData.CreationTime > BeforeDateTime)
+			{
+				NewFiles.Add(NormalizedFilename);
+			}
+			else if (StatData.ModificationTime > BeforeDateTime)
+			{
+				ModifiedFiles.Add(NormalizedFilename);
+			}
+			RemovedLongPackageNames.Remove(FName(FilePackagePath.ToView()));
+
+			return true;
+		});
+
+	for (FName LongPackageName : RemovedLongPackageNames)
+	{
+		// This file was deleted. Remove all assets in the package from the registry.
+		RemovePackageData(EventContext, LongPackageName);
+		// If the package was a package we were tracking as empty (due to e.g. a rename in editor), remove it.
+		// Disk now matches editor
+		RemoveEmptyPackage(LongPackageName);
+	}
+	if (NewFiles.Num() || NewDirs.Num())
+	{
+		if (GlobalGatherer.IsValid())
+		{
+			for (FString& NewDir : NewDirs)
+			{
+				GlobalGatherer->OnDirectoryCreated(NewDir);
+			}
+			GlobalGatherer->OnFilesCreated(NewFiles);
+			if (GlobalGatherer->IsSynchronous())
+			{
+				Impl::FScanPathContext Context(EventContext, InheritanceContext, NewDirs, NewFiles,
+					false /* bForceRescan */, false /* bIgnoreDenyListScanFilters */, nullptr /* OutFoundAssets */);
+				ScanPathsSynchronous(Context);
+			}
+		}
+	}
+	ScanModifiedAssetFiles(EventContext, InheritanceContext, ModifiedFiles);
+}
+
 }
 
 void UAssetRegistryImpl::OnAssetLoaded(UObject *AssetLoaded)
diff --git a/Engine/Source/Runtime/AssetRegistry/Private/AssetRegistryImpl.h b/Engine/Source/Runtime/AssetRegistry/Private/AssetRegistryImpl.h
--- a/Engine/Source/Runtime/AssetRegistry/Private/AssetRegistryImpl.h
+++ b/Engine/Source/Runtime/AssetRegistry/Private/AssetRegistryImpl.h
@@ -155,6 +155,8 @@ public:
 #if WITH_EDITOR
 	void OnDirectoryChanged(Impl::FEventContext& EventContext, Impl::FClassInheritanceContext& InheritanceContext,
 		TArray<FFileChangeData>& FileChangesProcessed);
+	void OnDirectoryRescanRequired(Impl::FEventContext& EventContext,
+		Impl::FClassInheritanceContext& InheritanceContext, FString& DirPath, int64 BeforeTimeStamp);
 	void AddLoadedAssetToProcess(const UObject& AssetLoaded);
 #endif
 	void OnContentPathMounted(Impl::FEventContext& EventContext, Impl::FClassInheritanceContext& InheritanceContext,
