DirectoryWatcher: Fixed buffer overflow not resulting in a FCA_RescanRequired.

diff --git a/Engine/Source/Developer/DirectoryWatcher/Private/Windows/DirectoryWatchRequestWindows.cpp b/Engine/Source/Developer/DirectoryWatcher/Private/Windows/DirectoryWatchRequestWindows.cpp
--- a/Engine/Source/Developer/DirectoryWatcher/Private/Windows/DirectoryWatchRequestWindows.cpp
+++ b/Engine/Source/Developer/DirectoryWatcher/Private/Windows/DirectoryWatchRequestWindows.cpp
@@ -196,13 +196,20 @@ void FDirectoryWatchRequestWindows::ProcessChange(uint32 Error, uint32 NumBytes)
 		bPendingDelete = true;
 	};
 
+	if (Error == 0 && NumBytes == 0)
+	{
+		DWORD UnusedNumberOfBytes;
+		GetOverlappedResult(DirectoryHandle, &Overlapped, &UnusedNumberOfBytes, 0);
+		Error = ::GetLastError();
+	}
+
 	if (Error == ERROR_OPERATION_ABORTED) 
 	{
 		// The operation was aborted, likely due to EndWatchRequest canceling it.
 		// Mark the request for delete so it can be cleaned up next tick.
 		bPendingDelete = true;
 		UE_CLOG(!IsEngineExitRequested(), LogDirectoryWatcher, Log, TEXT("A directory notification for '%s' was aborted."), *Directory);
-		return; 
+		return;
 	}
 	bool bValidNotification = Error != ERROR_IO_INCOMPLETE && NumBytes > 0;
 	bool bIsRescan = false;
