Fix directory watcher not re-adding assets when remove and rescan event is generated

diff --git a/Engine/Source/Developer/DirectoryWatcher/Private/Windows/DirectoryWatchRequestWindows.cpp b/Engine/Source/Developer/DirectoryWatcher/Private/Windows/DirectoryWatchRequestWindows.cpp
--- a/Engine/Source/Developer/DirectoryWatcher/Private/Windows/DirectoryWatchRequestWindows.cpp
+++ b/Engine/Source/Developer/DirectoryWatcher/Private/Windows/DirectoryWatchRequestWindows.cpp
@@ -9,7 +9,8 @@ FDirectoryWatchRequestWindows::FDirectoryWatchRequestWindows(uint32 Flags)
 {
 	bPendingDelete = false;
 	bEndWatchRequestInvoked = false;
-	WatchStartedTimestamp = 0;
+	WatchStartedTimeStampHistory[0] = 0;
+	WatchStartedTimeStampHistory[1] = 0;
 
 	bWatchSubtree = (Flags & IDirectoryWatcher::WatchOptions::IgnoreChangesInSubtree) == 0;
 	bool bIncludeDirectoryEvents = (Flags & IDirectoryWatcher::WatchOptions::IncludeDirectoryChanges) != 0;
@@ -110,7 +111,8 @@ bool FDirectoryWatchRequestWindows::Init(const FString& InDirectory)
 		DirectoryHandle = INVALID_HANDLE_VALUE;
 		return false;
 	}
-	WatchStartedTimestamp = FDateTime::UtcNow().ToUnixTimestamp();
+	WatchStartedTimeStampHistory[0] = FDateTime::UtcNow().ToUnixTimestamp();
+	WatchStartedTimeStampHistory[1] = WatchStartedTimeStampHistory[0];
 	bBufferInUse = true;
 
 	return true;
@@ -213,7 +215,7 @@ void FDirectoryWatchRequestWindows::ProcessChange(uint32 Error, uint32 NumBytes)
 	}
 	bool bValidNotification = Error != ERROR_IO_INCOMPLETE && NumBytes > 0;
 	bool bIsRescan = false;
-	int64 PreviousWatchStartedTimeStamp = 0;
+	int64 RescanReportTimestamp = 0;
 	if (bValidNotification)
 	{
 		// Swap the pointer to the backbuffer so we can start a new read as soon as possible
@@ -245,7 +247,7 @@ void FDirectoryWatchRequestWindows::ProcessChange(uint32 Error, uint32 NumBytes)
 	{
 		bValidNotification = true;
 		bIsRescan = true;
-		PreviousWatchStartedTimeStamp = WatchStartedTimestamp;
+		RescanReportTimestamp = WatchStartedTimeStampHistory[1];
 	}
 	else if (Error != ERROR_SUCCESS)
 	{
@@ -275,7 +277,8 @@ void FDirectoryWatchRequestWindows::ProcessChange(uint32 Error, uint32 NumBytes)
 		CloseHandleAndMarkForDelete();
 		return;
 	}
-	WatchStartedTimestamp = FDateTime::UtcNow().ToUnixTimestamp();
+	WatchStartedTimeStampHistory[1] = WatchStartedTimeStampHistory[0];
+	WatchStartedTimeStampHistory[0] = FDateTime::UtcNow().ToUnixTimestamp();
 	bBufferInUse = true;
 
 	// No need to process the change if we can not execute any delegates
@@ -340,7 +343,7 @@ void FDirectoryWatchRequestWindows::ProcessChange(uint32 Error, uint32 NumBytes)
 	else
 	{
 		FFileChangeData& ChangeData = FileChanges.Emplace_GetRef(Directory, FFileChangeData::FCA_RescanRequired);
-		ChangeData.TimeStamp = PreviousWatchStartedTimeStamp;
+		ChangeData.TimeStamp = RescanReportTimestamp;
 	}
 }
 
diff --git a/Engine/Source/Developer/DirectoryWatcher/Private/Windows/DirectoryWatchRequestWindows.h b/Engine/Source/Developer/DirectoryWatcher/Private/Windows/DirectoryWatchRequestWindows.h
--- a/Engine/Source/Developer/DirectoryWatcher/Private/Windows/DirectoryWatchRequestWindows.h
+++ b/Engine/Source/Developer/DirectoryWatcher/Private/Windows/DirectoryWatchRequestWindows.h
@@ -43,7 +43,7 @@ private:
 
 	FString Directory;
 	HANDLE DirectoryHandle;
-	int64 WatchStartedTimestamp;
+	int64 WatchStartedTimeStampHistory[2];
 	uint32 NotifyFilter;
 	uint32 BufferLength = 0;
 	struct FDeleterFree
