Add Insights trace for FDirectoryWatcherWindows::Tick

#rnx

diff --git a/Engine/Source/Developer/DirectoryWatcher/Private/Windows/DirectoryWatcherWindows.cpp b/Engine/Source/Developer/DirectoryWatcher/Private/Windows/DirectoryWatcherWindows.cpp
--- a/Engine/Source/Developer/DirectoryWatcher/Private/Windows/DirectoryWatcherWindows.cpp
+++ b/Engine/Source/Developer/DirectoryWatcher/Private/Windows/DirectoryWatcherWindows.cpp
@@ -110,6 +110,8 @@ bool FDirectoryWatcherWindows::UnregisterDirectoryChangedCallback_Handle( const
 
 void FDirectoryWatcherWindows::Tick( float DeltaSeconds )
 {
+	TRACE_CPUPROFILER_EVENT_SCOPE(FDirectoryWatcherWindows::Tick);
+
 	TArray<HANDLE> DirectoryHandles;
 	TMap<FDirectoryWithFlags, FDirectoryWatchRequestWindows*> InvalidRequestsToDelete;
 
