AssetRegistry+DirectoryWatcher: Fix indefinite delay of directory change notifications (and therefore new asset notifications) from FDirectoryWatcherWindows when a dialog box is active. MsgWaitForMultipleObjects is both slower then SleepEx for the processing of notifications, and suffers from the indefinite delay which SleepEx does not.
#rnx

diff --git a/Engine/Source/Developer/DirectoryWatcher/Private/Windows/DirectoryWatcherWindows.cpp b/Engine/Source/Developer/DirectoryWatcher/Private/Windows/DirectoryWatcherWindows.cpp
--- a/Engine/Source/Developer/DirectoryWatcher/Private/Windows/DirectoryWatcherWindows.cpp
+++ b/Engine/Source/Developer/DirectoryWatcher/Private/Windows/DirectoryWatcherWindows.cpp
@@ -111,58 +111,45 @@ bool FDirectoryWatcherWindows::UnregisterDirectoryChangedCallback_Handle( const
 void FDirectoryWatcherWindows::Tick( float DeltaSeconds )
 {
 	TRACE_CPUPROFILER_EVENT_SCOPE(FDirectoryWatcherWindows::Tick);
+	// Iterate over all requests in the request map. Move any that have shutdown into the pending delete
+	// list. If we find any that are not shutdown, poll for notifications (which call our callbacks) from the
+	// OS before we process the callbacks we received.
+	bool bHasPolled = false;
 
-	TArray<HANDLE> DirectoryHandles;
-	TMap<FDirectoryWithFlags, FDirectoryWatchRequestWindows*> InvalidRequestsToDelete;
-
-	// Find all handles to listen to and invalid requests to delete
-	for (TMap<FDirectoryWithFlags, FDirectoryWatchRequestWindows*>::TConstIterator RequestIt(RequestMap); RequestIt; ++RequestIt)
+	for (TMap<FDirectoryWithFlags, FDirectoryWatchRequestWindows*>::TIterator RequestIt(RequestMap);
+		RequestIt; ++RequestIt)
 	{
 		if ( RequestIt.Value()->IsPendingDelete() )
 		{
-			InvalidRequestsToDelete.Add(RequestIt.Key(), RequestIt.Value());
+			RequestsPendingDelete.AddUnique(RequestIt.Value());
+			RequestIt.RemoveCurrent();
 		}
 		else
 		{
-			DirectoryHandles.Add(RequestIt.Value()->GetDirectoryHandle());
-		}
-	}
-
-	// Remove all invalid requests from the request map and add them to the pending delete list so they will be deleted below
-	for (TMap<FDirectoryWithFlags, FDirectoryWatchRequestWindows*>::TConstIterator RequestIt(InvalidRequestsToDelete); RequestIt; ++RequestIt)
-	{
-		RequestMap.Remove(RequestIt.Key());
-		RequestsPendingDelete.AddUnique(RequestIt.Value());
-	}
-
-	// Trigger any file changed delegates that are queued up
-	// We need to do this in batches of MAXIMUM_WAIT_OBJECTS-1 as described by the documentation for MsgWaitForMultipleObjectsEx
-	if ( DirectoryHandles.Num() > 0 )
-	{
-		static const int32 BatchSize = MAXIMUM_WAIT_OBJECTS - 1;
-		for (int32 HandleOffset = 0; HandleOffset < DirectoryHandles.Num(); HandleOffset += BatchSize)
-		{
-			MsgWaitForMultipleObjectsEx(FMath::Min(DirectoryHandles.Num() - HandleOffset, BatchSize), DirectoryHandles.GetData() + HandleOffset, 0, QS_ALLEVENTS, MWMO_ALERTABLE);
+			if (!bHasPolled)
+			{
+				// Trigger any OS notifications that are queued up in Windows's MessageQueue. To do this we need to put the
+				// current thread into an alertable state, which we can do by calling SleepEx(0).
+				SleepEx(0 /* TimeoutMs. 0 means poll notifies and yield to waiting threads but otherwise return immediately */,
+						1 /* bAlertable: the notifications will be processed */);
+				bHasPolled = true;
+			}
+			// Pass on the unreal notifications for the request if it received an OS notification.
+			RequestIt.Value()->ProcessPendingNotifications();
 		}
 	}
 
-	// Delete any stale or invalid requests
-	for ( int32 RequestIdx = RequestsPendingDelete.Num() - 1; RequestIdx >= 0; --RequestIdx )
+	// Delete any stale or invalid requests, either the ones that we found during tick or ones that were
+	// queued for deletion in between ticks.
+	for (TArray<FDirectoryWatchRequestWindows*>::TIterator Iter(RequestsPendingDelete); Iter; ++Iter)
 	{
-		FDirectoryWatchRequestWindows* Request = RequestsPendingDelete[RequestIdx];
-
-		if ( Request->IsPendingDelete() )
+		FDirectoryWatchRequestWindows* Request = *Iter;
+		if (Request->IsPendingDelete())
 		{
 			// This request is safe to delete. Delete and remove it from the list
 			delete Request;
 			NumRequests--;
-			RequestsPendingDelete.RemoveAt(RequestIdx);
+			Iter.RemoveCurrentSwap();
 		}
 	}
-
-	// Finally, trigger any file change notification delegates
-	for (TMap<FDirectoryWithFlags, FDirectoryWatchRequestWindows*>::TConstIterator RequestIt(RequestMap); RequestIt; ++RequestIt)
-	{
-		RequestIt.Value()->ProcessPendingNotifications();
-	}
 }
