- Hopeful fix for some threading issues
- Will need to go to 5.4

diff --git a/Engine/Source/Developer/DirectoryWatcher/Private/Mac/DirectoryWatchRequestMac.cpp b/Engine/Source/Developer/DirectoryWatcher/Private/Mac/DirectoryWatchRequestMac.cpp
--- a/Engine/Source/Developer/DirectoryWatcher/Private/Mac/DirectoryWatchRequestMac.cpp
+++ b/Engine/Source/Developer/DirectoryWatcher/Private/Mac/DirectoryWatchRequestMac.cpp
@@ -11,8 +11,13 @@ void DirectoryWatchMacCallback( ConstFSEventStreamRef StreamRef, void* WatchRequ
 	check(WatchRequest);
 	check(WatchRequest->EventStream == StreamRef);
 
+	// make sure the array isn't destroyed until we are done with it on the other thread
+	CFArrayRef EventPathArray = (CFArrayRef)EventPaths;
+	CFRetain(EventPathArray);
+
 	GameThreadCall(^{
 		WatchRequest->ProcessChanges( EventCount, EventPaths, EventFlags);
+		CFRelease(EventPathArray);
 	});
 }
 
