- Changed dir watcher on Mac to use an FRWLock protecting an array instead of GameThreadCall that was causing trouble when cooking

diff --git a/Engine/Source/Developer/DirectoryWatcher/Private/Mac/DirectoryWatchRequestMac.cpp b/Engine/Source/Developer/DirectoryWatcher/Private/Mac/DirectoryWatchRequestMac.cpp
--- a/Engine/Source/Developer/DirectoryWatcher/Private/Mac/DirectoryWatchRequestMac.cpp
+++ b/Engine/Source/Developer/DirectoryWatcher/Private/Mac/DirectoryWatchRequestMac.cpp
@@ -11,14 +11,7 @@ void DirectoryWatchMacCallback( ConstFSEventStreamRef StreamRef, void* WatchRequ
 	check(WatchRequest);
 	check(WatchRequest->EventStream == StreamRef);
 
-	// make sure the array isn't destroyed until we are done with it on the other thread
-	CFArrayRef EventPathArray = (CFArrayRef)EventPaths;
-	CFRetain(EventPathArray);
-
-	GameThreadCall(^{
-		WatchRequest->ProcessChanges( EventCount, EventPaths, EventFlags);
-		CFRelease(EventPathArray);
-	});
+	WatchRequest->ProcessChanges( EventCount, EventPaths, EventFlags);
 }
 
 // ============================================================================================================================
@@ -129,34 +122,45 @@ void FDirectoryWatchRequestMac::EndWatchRequest()
 
 void FDirectoryWatchRequestMac::ProcessPendingNotifications()
 {
-	// Trigger all listening delegates with the files that have changed
-	if ( FileChanges.Num() > 0 )
+	bool bNeedsEmpty = false;
+
 	{
-		TMap<uint32, TArray<FFileChangeData>> FileChangeCache;
-		for (const FWatchDelegate& Delegate : Delegates)
+		FReadScopeLock Lock(FileChangesLock);
+
+		// Trigger all listening delegates with the files that have changed
+		if ( FileChanges.Num() > 0 )
 		{
-			// Filter list of all file changes down to ones that just match this delegate's flags
-			TArray<FFileChangeData>* CachedChanges = FileChangeCache.Find(Delegate.Value);
-			if (CachedChanges)
-			{
-				Delegate.Key.Execute(*CachedChanges);
-			}
-			else
+			TMap<uint32, TArray<FFileChangeData>> FileChangeCache;
+			for (const FWatchDelegate& Delegate : Delegates)
 			{
-				const bool bIncludeDirs = (Delegate.Value & IDirectoryWatcher::WatchOptions::IncludeDirectoryChanges) != 0;
-				TArray<FFileChangeData>& Changes = FileChangeCache.Add(Delegate.Value);
-				for (const TPair<FFileChangeData, bool>& FileChangeData : FileChanges)
+				// Filter list of all file changes down to ones that just match this delegate's flags
+				TArray<FFileChangeData>* CachedChanges = FileChangeCache.Find(Delegate.Value);
+				if (CachedChanges)
 				{
-					// @todo support IgnoreChangesInSubtree
-					if (!FileChangeData.Value || bIncludeDirs)
+					Delegate.Key.Execute(*CachedChanges);
+				}
+				else
+				{
+					const bool bIncludeDirs = (Delegate.Value & IDirectoryWatcher::WatchOptions::IncludeDirectoryChanges) != 0;
+					TArray<FFileChangeData>& Changes = FileChangeCache.Add(Delegate.Value);
+					for (const TPair<FFileChangeData, bool>& FileChangeData : FileChanges)
 					{
-						Changes.Add(FileChangeData.Key);
+						// @todo support IgnoreChangesInSubtree
+						if (!FileChangeData.Value || bIncludeDirs)
+						{
+							Changes.Add(FileChangeData.Key);
+						}
 					}
+					Delegate.Key.Execute(Changes);
 				}
-				Delegate.Key.Execute(Changes);
 			}
+			
+			bNeedsEmpty = true;
 		}
-
+	}
+	if (bNeedsEmpty)
+	{
+		FWriteScopeLock Lock(FileChangesLock);
 		FileChanges.Empty();
 	}
 }
@@ -169,6 +173,8 @@ void FDirectoryWatchRequestMac::ProcessChanges( size_t EventCount, void* EventPa
 		return;
 	}
 
+	FWriteScopeLock Lock(FileChangesLock);
+
 	CFArrayRef EventPathArray = (CFArrayRef)EventPaths;
 
 	for( size_t EventIndex = 0; EventIndex < EventCount; ++EventIndex )
diff --git a/Engine/Source/Developer/DirectoryWatcher/Private/Mac/DirectoryWatchRequestMac.h b/Engine/Source/Developer/DirectoryWatcher/Private/Mac/DirectoryWatchRequestMac.h
--- a/Engine/Source/Developer/DirectoryWatcher/Private/Mac/DirectoryWatchRequestMac.h
+++ b/Engine/Source/Developer/DirectoryWatcher/Private/Mac/DirectoryWatchRequestMac.h
@@ -7,6 +7,7 @@
 #include "Containers/UnrealString.h"
 #include "Containers/Array.h"
 #include <CoreServices/CoreServices.h>
+#include "Misc/ScopeRWLock.h"
 
 class FDirectoryWatchRequestMac
 {
@@ -39,6 +40,7 @@ private:
 	TArray<FWatchDelegate> Delegates;
 	/** Each FFileChangeData tracks whether it is a directory or not */
 	TArray<TPair<FFileChangeData, bool>> FileChanges;
+	FRWLock FileChangesLock;
 
 	friend void DirectoryWatchMacCallback( ConstFSEventStreamRef StreamRef, void* WatchRequestPtr, size_t EventCount, void* EventPaths, const FSEventStreamEventFlags EventFlags[], const FSEventStreamEventId EventIDs[] );
 
