[FileCache]
 - Add Insight scopes

#rnx

diff --git a/Engine/Source/Developer/DirectoryWatcher/Private/FileCache.cpp b/Engine/Source/Developer/DirectoryWatcher/Private/FileCache.cpp
--- a/Engine/Source/Developer/DirectoryWatcher/Private/FileCache.cpp
+++ b/Engine/Source/Developer/DirectoryWatcher/Private/FileCache.cpp
@@ -279,6 +279,7 @@ FAsyncDirectoryReader::EProgressResult FAsyncDirectoryReader::Tick(const FTimeLi
 		return EProgressResult::Finished;
 	}
 
+	TRACE_CPUPROFILER_EVENT_SCOPE(FAsyncDirectoryReader::Tick);
 	auto& FileManager = IFileManager::Get();
 
 	// Discover files
@@ -339,6 +340,8 @@ FAsyncDirectoryReader::EProgressResult FAsyncDirectoryReader::Tick(const FTimeLi
 
 void FAsyncDirectoryReader::ScanDirectory(const FString& InDirectory)
 {
+	TRACE_CPUPROFILER_EVENT_SCOPE(FAsyncDirectoryReader::ScanDirectory);
+
 	struct FVisitor : public IPlatformFile::FDirectoryVisitor
 	{
 		TArray<FString>* PendingFiles;
