Updated headers using UnrealCodeFixup to make sure dllstorage is on methods/staticvars instead of types. Part 4/n

diff --git a/Engine/Source/Developer/BSPUtils/Public/BSPUtils.h b/Engine/Source/Developer/BSPUtils/Public/BSPUtils.h
--- a/Engine/Source/Developer/BSPUtils/Public/BSPUtils.h
+++ b/Engine/Source/Developer/BSPUtils/Public/BSPUtils.h
@@ -5,53 +5,55 @@
 #include "CoreMinimal.h"
 #include "Engine/Brush.h"
 
+#define UE_API BSPUTILS_API
+
 class UWorld;
 class UModel;
 class UMaterialInterface;
 class FPoly;
 
 
-class BSPUTILS_API FBSPUtils
+class FBSPUtils
 {
 public:
 	/** Repartition Bsp tree */
-	static void bspRepartition( UWorld* InWorld, int32 iNode );
+	static UE_API void bspRepartition( UWorld* InWorld, int32 iNode );
 
 	/** Convert a Bsp node to an EdPoly.  Returns number of vertices in Bsp node. */
-	static int32 bspNodeToFPoly( UModel* Model, int32 iNode, FPoly* EdPoly );
+	static UE_API int32 bspNodeToFPoly( UModel* Model, int32 iNode, FPoly* EdPoly );
 
 	/**
 	 * Clean up all nodes after a CSG operation.  Resets temporary bit flags and unlinks
 	 * empty leaves.  Removes zero-vertex nodes which have nonzero-vertex coplanars.
 	 */
-	static void bspCleanup( UModel* Model );
+	static UE_API void bspCleanup( UModel* Model );
 
 	/** 
 	 * Build EdPoly list from a model's Bsp. Not transactional.
 	 * @param DestArray helps build bsp FPolys in non-main threads. It also allows to perform this action without GUndo 
 	 *	      interfering. Temporary results will be written to DestArray. Defaults to Model->Polys->Element
 	 */
-	static void bspBuildFPolys( UModel* Model, bool SurfLinks, int32 iNode, TArray<FPoly>* DestArray = NULL );
+	static UE_API void bspBuildFPolys( UModel* Model, bool SurfLinks, int32 iNode, TArray<FPoly>* DestArray = NULL );
 
-	static void bspMergeCoplanars( UModel* Model, bool RemapLinks, bool MergeDisparateTextures );
+	static UE_API void bspMergeCoplanars( UModel* Model, bool RemapLinks, bool MergeDisparateTextures );
 
 	/**
 	 * Performs any CSG operation between the brush and the world.
 	 */
-	static int32 bspBrushCSG( ABrush* Actor, UModel* Model, UModel* TempModel, UMaterialInterface* SelectedMaterialInstance, uint32 PolyFlags, EBrushType BrushType, ECsgOper CSGOper, bool bBuildBounds, bool bMergePolys, bool bReplaceNULLMaterialRefs, bool bShowProgressBar=true );
+	static UE_API int32 bspBrushCSG( ABrush* Actor, UModel* Model, UModel* TempModel, UMaterialInterface* SelectedMaterialInstance, uint32 PolyFlags, EBrushType BrushType, ECsgOper CSGOper, bool bBuildBounds, bool bMergePolys, bool bReplaceNULLMaterialRefs, bool bShowProgressBar=true );
 
 	/**
 	 * Optimize a level's Bsp, eliminating T-joints where possible, and building side
 	 * links.  This does not always do a 100% perfect job, mainly due to imperfect 
 	 * levels, however it should never fail or return incorrect results.
 	 */
-	static void bspOptGeom( UModel* Model );
+	static UE_API void bspOptGeom( UModel* Model );
 
 	/**
 	 * Sets and clears all Bsp node flags.  Affects all nodes, even ones that don't
 	 * really exist.
 	 */
-	static void polySetAndClearPolyFlags(UModel *Model, uint32 SetBits, uint32 ClearBits,bool SelectedOnly, bool UpdateBrush);
+	static UE_API void polySetAndClearPolyFlags(UModel *Model, uint32 SetBits, uint32 ClearBits,bool SelectedOnly, bool UpdateBrush);
 
 	/**
 	 *
@@ -63,10 +65,10 @@ public:
 	 *
 	 * returns true if poly not available
 	 */
-	static bool polyFindBrush(UModel* InModel, int32 iSurf, FPoly &Poly);
+	static UE_API bool polyFindBrush(UModel* InModel, int32 iSurf, FPoly &Poly);
 
 	UE_DEPRECATED(5.1, "polyFindMaster is deprecated; please use polyFindBrush instead")
-	static bool polyFindMaster(UModel* InModel, int32 iSurf, FPoly &Poly);
+	static UE_API bool polyFindMaster(UModel* InModel, int32 iSurf, FPoly &Poly);
 
 	/**
 	 * Update the brush EdPoly corresponding to a newly-changed
@@ -74,26 +76,28 @@ public:
 	 *
 	 * Doesn't do any transaction tracking.
 	 */
-	static void polyUpdateBrush(UModel* Model, int32 iSurf, bool bUpdateTexCoords, bool bOnlyRefreshSurfaceMaterials);
+	static UE_API void polyUpdateBrush(UModel* Model, int32 iSurf, bool bUpdateTexCoords, bool bOnlyRefreshSurfaceMaterials);
 	
 	UE_DEPRECATED(5.1, "polyUpdateMaster is deprecated; please use polyUpdateBrush instead")
-	static void polyUpdateMaster(UModel* Model, int32 iSurf, bool bUpdateTexCoords, bool bOnlyRefreshSurfaceMaterials);
+	static UE_API void polyUpdateMaster(UModel* Model, int32 iSurf, bool bUpdateTexCoords, bool bOnlyRefreshSurfaceMaterials);
 	
 	/**
 	 * Populates a list with all polys that are linked to the specified poly.  The
 	 * resulting list includes the original poly.
 	 */
-	static void polyGetLinkedPolys(ABrush* InBrush, FPoly* InPoly, TArray<FPoly>* InPolyList);
+	static UE_API void polyGetLinkedPolys(ABrush* InBrush, FPoly* InPoly, TArray<FPoly>* InPolyList);
 
 	/**
 	 * Takes a list of polygons and creates a new list of polys which have no overlapping edges.  It splits
 	 * edges as necessary to achieve this.
 	 */
-	static void polySplitOverlappingEdges( TArray<FPoly>* InPolyList, TArray<FPoly>* InResult );
+	static UE_API void polySplitOverlappingEdges( TArray<FPoly>* InPolyList, TArray<FPoly>* InResult );
 
 	/**
 	 * Takes a list of polygons and returns a list of the outside edges (edges which are not shared
 	 * by other polys in the list).
 	 */
-	static void polyGetOuterEdgeList(TArray<FPoly>* InPolyList, TArray<FEdge>* InEdgeList);
+	static UE_API void polyGetOuterEdgeList(TArray<FPoly>* InPolyList, TArray<FEdge>* InEdgeList);
 };
+
+#undef UE_API
diff --git a/Engine/Source/Developer/CookedEditor/Public/CookedEditorPackageManager.h b/Engine/Source/Developer/CookedEditor/Public/CookedEditorPackageManager.h
--- a/Engine/Source/Developer/CookedEditor/Public/CookedEditorPackageManager.h
+++ b/Engine/Source/Developer/CookedEditor/Public/CookedEditorPackageManager.h
@@ -63,14 +63,14 @@ typedef FMacEditorTargetPlatformControlsParent FHostPlatformEditorTargetPlatform
  * You likely can use just this class without needing to subclass TCookedEditorTargetPlatform at all, but you can do that
  * if needed.
  */
-class COOKEDEDITOR_API ICookedEditorPackageManager
+class ICookedEditorPackageManager
 {
 public:
 
 	/** 
 	 * Construct a package manager for the given TP
 	 */
-	static TUniquePtr<ICookedEditorPackageManager> FactoryForTargetPlatform(bool bIsCookedCooker);
+	static COOKEDEDITOR_API TUniquePtr<ICookedEditorPackageManager> FactoryForTargetPlatform(bool bIsCookedCooker);
 
 	virtual ~ICookedEditorPackageManager()
 	{
@@ -134,19 +134,19 @@ protected:
 	 * @param Path rooted path to look in ("/Engine", etc)
 	 * @param SeachMode Recurse or TopDirectoryOnly
 	 */
-	void AddPackagesFromPath(TArray<FName>& Packages, const TCHAR* Path, EPackageSearchMode SearchMode) const;
+	COOKEDEDITOR_API void AddPackagesFromPath(TArray<FName>& Packages, const TCHAR* Path, EPackageSearchMode SearchMode) const;
 
 	/**
 	 * Allow a subclass to remove packages found in GatherAllPackages
 	 */
-	virtual void FilterGatheredPackages(TArray<FName>& PackageNames) const;
+	COOKEDEDITOR_API virtual void FilterGatheredPackages(TArray<FName>& PackageNames) const;
 
 protected:
 
 	/**
 	 * Meat of this class, this calls other functions that generally will be overridden - subclass needs to pass in disabled plugins
 	 */
-	void GatherAllPackagesExceptDisabled(TArray<FName>& PackageNames, const TArray<FString>& DisabledPlugins) const;
+	COOKEDEDITOR_API void GatherAllPackagesExceptDisabled(TArray<FName>& PackageNames, const TArray<FString>& DisabledPlugins) const;
 };
 
 
@@ -155,7 +155,7 @@ protected:
  * An implentation of ICookedEditorPackageManager that uses BaseGame.ini / DefaultGame.ini to control cooked editor packaging setup
  * 
  */
-class COOKEDEDITOR_API FIniCookedEditorPackageManager : public ICookedEditorPackageManager
+class FIniCookedEditorPackageManager : public ICookedEditorPackageManager
 {
 	TArray<FString> EngineAssetPaths;
 	TArray<FString> ProjectAssetPaths;
@@ -174,25 +174,25 @@ class COOKEDEDITOR_API FIniCookedEditorPackageManager : public ICookedEditorPack
 	int32 CookProgressDisplayMode = 0;
 
 	// gets an array from two sections, depending on bIsCookedCooker setting
-	TArray<FString> GetConfigArray(const TCHAR* Key) const;
+	COOKEDEDITOR_API TArray<FString> GetConfigArray(const TCHAR* Key) const;
 
 	// Load the classes specified from config settings that we need to reference
-	void InitializeClasses();
+	COOKEDEDITOR_API void InitializeClasses();
 public:
 
-	FIniCookedEditorPackageManager(bool bIsCookedCooker);
+	COOKEDEDITOR_API FIniCookedEditorPackageManager(bool bIsCookedCooker);
 
 	// gets an array from two sections, depending on bIsCookedCooker setting
-	static TArray<FString> GetConfigArray(const TCHAR* Key, bool bIsCookedCooker);
-
-	virtual void FilterGatheredPackages(TArray<FName>& PackageNames) const override;
-	virtual void InitializeForCook() override;
-	virtual void GetEnginePackagesToCook(TArray<FName>& PackagesToCook) const override;
-	virtual void GetProjectPackagesToCook(TArray<FName>& PackagesToCook) const override;
-	virtual bool AllowObjectToBeCooked(const class UObject* Obj) const override;
-	virtual bool AllowAssetToBeGathered(const struct FAssetData& AssetData) const override;
-	virtual bool AllowEnginePluginContentToBeCooked(const TSharedRef<IPlugin>) const override;
-	virtual bool AllowProjectPluginContentToBeCooked(const TSharedRef<IPlugin>) const override;
+	static COOKEDEDITOR_API TArray<FString> GetConfigArray(const TCHAR* Key, bool bIsCookedCooker);
+
+	COOKEDEDITOR_API virtual void FilterGatheredPackages(TArray<FName>& PackageNames) const override;
+	COOKEDEDITOR_API virtual void InitializeForCook() override;
+	COOKEDEDITOR_API virtual void GetEnginePackagesToCook(TArray<FName>& PackagesToCook) const override;
+	COOKEDEDITOR_API virtual void GetProjectPackagesToCook(TArray<FName>& PackagesToCook) const override;
+	COOKEDEDITOR_API virtual bool AllowObjectToBeCooked(const class UObject* Obj) const override;
+	COOKEDEDITOR_API virtual bool AllowAssetToBeGathered(const struct FAssetData& AssetData) const override;
+	COOKEDEDITOR_API virtual bool AllowEnginePluginContentToBeCooked(const TSharedRef<IPlugin>) const override;
+	COOKEDEDITOR_API virtual bool AllowProjectPluginContentToBeCooked(const TSharedRef<IPlugin>) const override;
 	
-	virtual void GatherAllPackages(TArray<FName>& PackageNames) const override;
+	COOKEDEDITOR_API virtual void GatherAllPackages(TArray<FName>& PackageNames) const override;
 };
diff --git a/Engine/Source/Developer/DesktopWidgets/Public/Widgets/Input/SDirectoryPicker.h b/Engine/Source/Developer/DesktopWidgets/Public/Widgets/Input/SDirectoryPicker.h
--- a/Engine/Source/Developer/DesktopWidgets/Public/Widgets/Input/SDirectoryPicker.h
+++ b/Engine/Source/Developer/DesktopWidgets/Public/Widgets/Input/SDirectoryPicker.h
@@ -8,12 +8,14 @@
 #include "Widgets/SCompoundWidget.h"
 #include "Application/SlateWindowHelper.h"
 
+#define UE_API DESKTOPWIDGETS_API
+
 class SEditableTextBox;
 
 /**
  * A directory path box, with a button for picking a new path.
  */
-class DESKTOPWIDGETS_API SDirectoryPicker : public SCompoundWidget
+class SDirectoryPicker : public SCompoundWidget
 {
 public:
 	DECLARE_DELEGATE_OneParam(FOnDirectoryChanged, const FString& /*Directory*/);
@@ -33,9 +35,9 @@ public:
 	
 public:
 
-	void Construct(const FArguments& InArgs);
+	UE_API void Construct(const FArguments& InArgs);
 	
-	FString GetFilePath() const;
+	UE_API FString GetFilePath() const;
 
 	FString GetDirectory() const { return Directory; }
 
@@ -47,14 +49,14 @@ public:
 	DECLARE_DELEGATE_OneParam(FOnDirectoryPicked, const FString& /*PickedPath*/);
 	
 private:
-	void OnDirectoryTextChanged(const FText& InDirectoryPath);
-	void OnDirectoryTextCommited(const FText& InText, ETextCommit::Type InCommitType);
+	UE_API void OnDirectoryTextChanged(const FText& InDirectoryPath);
+	UE_API void OnDirectoryTextCommited(const FText& InText, ETextCommit::Type InCommitType);
 
-	FText GetFilePathText() const;
+	UE_API FText GetFilePathText() const;
 
-	bool OpenPlatformDirectoryPicker(FString& OutDirectory, const FString& DefaultPath);
+	UE_API bool OpenPlatformDirectoryPicker(FString& OutDirectory, const FString& DefaultPath);
 
-	FReply BrowseForDirectory();
+	UE_API FReply BrowseForDirectory();
 
 private:
 	FString File;
@@ -66,3 +68,5 @@ private:
 
 	TSharedPtr<SEditableTextBox> EditableTextBox;
 };
+
+#undef UE_API
diff --git a/Engine/Source/Developer/DesktopWidgets/Public/Widgets/Input/SFilePathPicker.h b/Engine/Source/Developer/DesktopWidgets/Public/Widgets/Input/SFilePathPicker.h
--- a/Engine/Source/Developer/DesktopWidgets/Public/Widgets/Input/SFilePathPicker.h
+++ b/Engine/Source/Developer/DesktopWidgets/Public/Widgets/Input/SFilePathPicker.h
@@ -12,6 +12,8 @@
 #include "Widgets/SCompoundWidget.h"
 #include "Styling/SlateTypes.h"
 
+#define UE_API DESKTOPWIDGETS_API
+
 class SEditableTextBox;
 
 /**
@@ -25,7 +27,7 @@ DECLARE_DELEGATE_OneParam(FOnPathPicked, const FString& /*PickedPath*/);
 /**
  * Implements an editable text box with a browse button.
  */
-class DESKTOPWIDGETS_API SFilePathPicker
+class SFilePathPicker
 	: public SCompoundWidget
 {
 public:
@@ -78,18 +80,18 @@ public:
 	 *
 	 * @param InArgs The construction arguments.
 	 */
-	void Construct( const FArguments& InArgs );
+	UE_API void Construct( const FArguments& InArgs );
 
 private:
 
 	/** Callback for clicking the browse button. */
-	FReply HandleBrowseButtonClicked( );
+	UE_API FReply HandleBrowseButtonClicked( );
 
 	/** Callback for getting the text in the path text box. */
-	FText HandleTextBoxText( ) const;
+	UE_API FText HandleTextBoxText( ) const;
 
 	/** Callback for committing the text in the path text box. */
-	void HandleTextBoxTextCommitted( const FText& NewText, ETextCommit::Type /*CommitInfo*/ );
+	UE_API void HandleTextBoxTextCommitted( const FText& NewText, ETextCommit::Type /*CommitInfo*/ );
 
 private:
 
@@ -116,3 +118,5 @@ private:
 	/** Holds a delegate that is executed when a file was picked. */
 	FOnPathPicked OnPathPicked;
 };
+
+#undef UE_API
diff --git a/Engine/Source/Developer/DirectoryWatcher/Public/FileCache.h b/Engine/Source/Developer/DirectoryWatcher/Public/FileCache.h
--- a/Engine/Source/Developer/DirectoryWatcher/Public/FileCache.h
+++ b/Engine/Source/Developer/DirectoryWatcher/Public/FileCache.h
@@ -25,6 +25,8 @@
 #include "Templates/SharedPointer.h"
 #include "Templates/UnrealTemplate.h"
 
+#define UE_API DIRECTORYWATCHER_API
+
 struct FFileChangeData;
 struct FGuid;
 
@@ -32,9 +34,9 @@ namespace DirectoryWatcher
 {
 
 /** Custom serialization version for FFileCache */
-struct DIRECTORYWATCHER_API FFileCacheCustomVersion
+struct FFileCacheCustomVersion
 {
-	static const FGuid Key;
+	static UE_API const FGuid Key;
 	enum Type {	Initial, IncludeFileHash, Latest = IncludeFileHash };
 };
 
@@ -411,54 +413,54 @@ private:
  * Changes in the cache can be retrieved through GetOutstandingChanges(). Changes will be reported for any
  * change in the cached state even between runs of the process.
  */
-class DIRECTORYWATCHER_API FFileCache
+class FFileCache
 {
 public:
 
 	/** Construction from a config */
-	FFileCache(const FFileCacheConfig& InConfig);
+	UE_API FFileCache(const FFileCacheConfig& InConfig);
 
 	FFileCache(const FFileCache&) = delete;
 	FFileCache& operator=(const FFileCache&) = delete;
 
 	/** Destructor */
-	~FFileCache();
+	UE_API ~FFileCache();
 
 	/** Destroy this cache. Cleans out the in-memory state and deletes the cache file, if present. */
-	void Destroy();
+	UE_API void Destroy();
 
 	/** Get the absolute path of the directory this cache reflects */
 	const FString& GetDirectory() const { return Config.Directory; }
 
 	/** Check whether this file cache has finished starting up yet. Does not imply rename/move detection is fully initialized. (see MoveDetectionInitialized()) */
-	bool HasStartedUp() const;
+	UE_API bool HasStartedUp() const;
 
 	/** Check whether this move/rename detection has been initiated or not. This can take much longer than startup, so can be checked separately */
-	bool MoveDetectionInitialized() const;
+	UE_API bool MoveDetectionInitialized() const;
 
 	/** Attempt to locate file data pertaining to the specified filename.
 	 *	@param		InFilename		The filename to find data for - either relative to the directory, or absolute, depending on Config.PathType.
 	 *	@return		The current cached file data pertaining to the specified filename, or nullptr. May not be completely up to date if there are outstanding transactions.
 	 */
-	const FFileData* FindFileData(FImmutableString InFilename) const;
+	UE_API const FFileData* FindFileData(FImmutableString InFilename) const;
 
 	/** Tick this FileCache */
-	void Tick();
+	UE_API void Tick();
 
 	/** Write out the cached file, if we have any changes to write */
-	void WriteCache();
+	UE_API void WriteCache();
 
 	/**
 	 * Return a transaction to the cache for completion. Will update the cached state with the change
 	 * described in the transaction, and mark the cache as needing to be saved.
 	 */
-	void CompleteTransaction(FUpdateCacheTransaction&& Transaction);
+	UE_API void CompleteTransaction(FUpdateCacheTransaction&& Transaction);
 
 	/** Report an external change to the manager, such that a subsequent equal change reported by the os be ignored */
-	void IgnoreNewFile(const FString& Filename);
-	void IgnoreFileModification(const FString& Filename);
-	void IgnoreMovedFile(const FString& SrcFilename, const FString& DstFilename);
-	void IgnoreDeletedFile(const FString& Filename);
+	UE_API void IgnoreNewFile(const FString& Filename);
+	UE_API void IgnoreFileModification(const FString& Filename);
+	UE_API void IgnoreMovedFile(const FString& SrcFilename, const FString& DstFilename);
+	UE_API void IgnoreDeletedFile(const FString& Filename);
 
 	/** Get the number of pending changes to the cache. */
 	int32 GetNumDirtyFiles() const { return DirtyFiles.Num(); }
@@ -469,51 +471,51 @@ public:
 								1) the pending transaction
 								2) the time the change was detected
 	*/
-	void IterateOutstandingChanges(TFunctionRef<bool(const FUpdateCacheTransaction&, const FDateTime&)> InPredicate) const;
+	UE_API void IterateOutstandingChanges(TFunctionRef<bool(const FUpdateCacheTransaction&, const FDateTime&)> InPredicate) const;
 
 	/** Get pending changes to the cache. Transactions must be returned to CompleteTransaction to update the cache.
 	 *  Filter predicate receives a transaction and the time the change was reported. */
-	TArray<FUpdateCacheTransaction> FilterOutstandingChanges(TFunctionRef<bool(const FUpdateCacheTransaction&, const FDateTime&)> InPredicate);
-	TArray<FUpdateCacheTransaction> GetOutstandingChanges();
+	UE_API TArray<FUpdateCacheTransaction> FilterOutstandingChanges(TFunctionRef<bool(const FUpdateCacheTransaction&, const FDateTime&)> InPredicate);
+	UE_API TArray<FUpdateCacheTransaction> GetOutstandingChanges();
 
 private:
 
 	/** Called when the directory we are monitoring has been changed in some way */
-	void OnDirectoryChanged(const TArray<FFileChangeData>& FileChanges);
+	UE_API void OnDirectoryChanged(const TArray<FFileChangeData>& FileChanges);
 
 	/** Diff the specified set of dirty files (absolute paths, or relative to the monitor directory), adding transactions to the specified array if necessary.
 	 *	Optionally takes a directory state from which we can retrieve current file system state, without having to ask the FS directly.
 	 *  Optionally exclude files that have changed since the specified threshold, to ensure that related events get grouped together correctly.
 	 */
-	void DiffDirtyFiles(TMap<FImmutableString, FFileData>& InDirtyFiles, TArray<FUpdateCacheTransaction>& OutTransactions, const FDirectoryState* InFileSystemState = nullptr) const;
+	UE_API void DiffDirtyFiles(TMap<FImmutableString, FFileData>& InDirtyFiles, TArray<FUpdateCacheTransaction>& OutTransactions, const FDirectoryState* InFileSystemState = nullptr) const;
 
 	/** Detect a rename for the specified file */
-	void DetectRename();
+	UE_API void DetectRename();
 
 	/** Get the absolute path from a transaction filename */
-	FString GetAbsolutePath(const FString& InTransactionPath) const;
+	UE_API FString GetAbsolutePath(const FString& InTransactionPath) const;
 
 	/** Get the transaction path from an absolute filename. Returns nothing if the absolute path is invalid (not under the correct folder, or not applicable) */
-	TOptional<FString> GetTransactionPath(const FString& InAbsolutePath) const;
+	UE_API TOptional<FString> GetTransactionPath(const FString& InAbsolutePath) const;
 
 	/** Unbind the watcher from the directory. Called on destruction. */
-	void UnbindWatcher();
+	UE_API void UnbindWatcher();
 
 	/** Read the cache file data and return the contents */
-	TOptional<FDirectoryState> ReadCache() const;
+	UE_API TOptional<FDirectoryState> ReadCache() const;
 
 	/** Called when the initial async reader has finished harvesting file system timestamps */
-	void ReadStateFromAsyncReader();
+	UE_API void ReadStateFromAsyncReader();
 
 	/** Called to harvest any file hashes that have been generated by the DirtyFileHasher thread task */
-	void HarvestDirtyFileHashes();
-	void RescanForDirtyFileHashes();
+	UE_API void HarvestDirtyFileHashes();
+	UE_API void RescanForDirtyFileHashes();
 
 	/** Update our cache of pending transactions (to prevent diffing against the directory unnecessarily) */
-	void UpdatePendingTransactions();
+	UE_API void UpdatePendingTransactions();
 
 	/** Returns version of path to store */
-	FString GetPathToStore(const FFilenameAndHash& InData) const;
+	UE_API FString GetPathToStore(const FFilenameAndHash& InData) const;
 
 private:
 
@@ -550,3 +552,5 @@ private:
 };
 
 } // namespace DirectoryWatcher
+
+#undef UE_API
diff --git a/Engine/Source/Developer/HierarchicalLODUtilities/Public/HierarchicalLODProxyProcessor.h b/Engine/Source/Developer/HierarchicalLODUtilities/Public/HierarchicalLODProxyProcessor.h
--- a/Engine/Source/Developer/HierarchicalLODUtilities/Public/HierarchicalLODProxyProcessor.h
+++ b/Engine/Source/Developer/HierarchicalLODUtilities/Public/HierarchicalLODProxyProcessor.h
@@ -13,14 +13,14 @@
 class ALODActor;
 class UHLODProxy;
 
-class HIERARCHICALLODUTILITIES_API FHierarchicalLODProxyProcessor : public FTSTickerObjectBase
+class FHierarchicalLODProxyProcessor : public FTSTickerObjectBase
 {
 public:		
-	FHierarchicalLODProxyProcessor();
-	~FHierarchicalLODProxyProcessor();
+	HIERARCHICALLODUTILITIES_API FHierarchicalLODProxyProcessor();
+	HIERARCHICALLODUTILITIES_API ~FHierarchicalLODProxyProcessor();
 	
 	/** Begin FTSTickerObjectBase */
-	virtual bool Tick(float DeltaTime) override;
+	HIERARCHICALLODUTILITIES_API virtual bool Tick(float DeltaTime) override;
 	/** End FTSTickerObjectBase */
 
 	/**
@@ -30,27 +30,27 @@ public:
 	* @param LODSetup - Simplification settings structure
 	* @return FGuid - Guid for the job
 	*/
-	FGuid AddProxyJob(ALODActor* InLODActor, UHLODProxy* InProxy, const FHierarchicalSimplification& LODSetup);
+	HIERARCHICALLODUTILITIES_API FGuid AddProxyJob(ALODActor* InLODActor, UHLODProxy* InProxy, const FHierarchicalSimplification& LODSetup);
 	
 	/** Callback function used for processing finished mesh generation jobs	
 	* @param InGuid - Guid of the finished job
 	* @param InAssetsToSync - Assets data created by the job
 	*/
-	void ProcessProxy(const FGuid InGuid, TArray<UObject*>& InAssetsToSync);
+	HIERARCHICALLODUTILITIES_API void ProcessProxy(const FGuid InGuid, TArray<UObject*>& InAssetsToSync);
 	
 	/** Returns the callback delegate which will be passed onto ProxyLOD function */
-	FCreateProxyDelegate& GetCallbackDelegate();
+	HIERARCHICALLODUTILITIES_API FCreateProxyDelegate& GetCallbackDelegate();
 		
-	bool IsProxyGenerationRunning() const;
+	HIERARCHICALLODUTILITIES_API bool IsProxyGenerationRunning() const;
 protected:
 	/** Called when the map has changed*/
-	void OnMapChange(uint32 MapFlags);
+	HIERARCHICALLODUTILITIES_API void OnMapChange(uint32 MapFlags);
 
 	/** Called when the current level has changed */
-	void OnNewCurrentLevel();
+	HIERARCHICALLODUTILITIES_API void OnNewCurrentLevel();
 
 	/** Clears the processing data array/map */
-	void ClearProcessingData();
+	HIERARCHICALLODUTILITIES_API void ClearProcessingData();
 protected:
 	/** Structure storing the data required during processing */
 	struct FProcessData
diff --git a/Engine/Source/Developer/HierarchicalLODUtilities/Public/HierarchicalLODUtilitiesModule.h b/Engine/Source/Developer/HierarchicalLODUtilities/Public/HierarchicalLODUtilitiesModule.h
--- a/Engine/Source/Developer/HierarchicalLODUtilities/Public/HierarchicalLODUtilitiesModule.h
+++ b/Engine/Source/Developer/HierarchicalLODUtilities/Public/HierarchicalLODUtilitiesModule.h
@@ -5,6 +5,8 @@
 #include "CoreMinimal.h"
 #include "Modules/ModuleInterface.h"
 
+#define UE_API HIERARCHICALLODUTILITIES_API
+
 class FHierarchicalLODProxyProcessor;
 class IHierarchicalLODUtilities;
 
@@ -19,16 +21,18 @@ public:
 /**
 * IHierarchicalLODUtilities module interface
 */
-class HIERARCHICALLODUTILITIES_API FHierarchicalLODUtilitiesModule : public IHierarchicalLODUtilitiesModule
+class FHierarchicalLODUtilitiesModule : public IHierarchicalLODUtilitiesModule
 {
 public:
-	virtual void ShutdownModule() override;
-	virtual void StartupModule() override;
+	UE_API virtual void ShutdownModule() override;
+	UE_API virtual void StartupModule() override;
 
 	/** Returns the Proxy processor instance from within this module */
-	virtual FHierarchicalLODProxyProcessor* GetProxyProcessor() override;
-	virtual IHierarchicalLODUtilities* GetUtilities() override;
+	UE_API virtual FHierarchicalLODProxyProcessor* GetProxyProcessor() override;
+	UE_API virtual IHierarchicalLODUtilities* GetUtilities() override;
 private:
 	FHierarchicalLODProxyProcessor* ProxyProcessor;
 	IHierarchicalLODUtilities* Utilities;
 };
+
+#undef UE_API
diff --git a/Engine/Source/Developer/HierarchicalLODUtilities/Public/IHierarchicalLODUtilities.h b/Engine/Source/Developer/HierarchicalLODUtilities/Public/IHierarchicalLODUtilities.h
--- a/Engine/Source/Developer/HierarchicalLODUtilities/Public/IHierarchicalLODUtilities.h
+++ b/Engine/Source/Developer/HierarchicalLODUtilities/Public/IHierarchicalLODUtilities.h
@@ -35,7 +35,7 @@ ENUM_CLASS_FLAGS(EClusterGenerationError);
 /**
  * IHierarchicalLODUtilities module interface
  */
-class HIERARCHICALLODUTILITIES_API IHierarchicalLODUtilities
+class IHierarchicalLODUtilities
 {
 public:	
 	virtual ~IHierarchicalLODUtilities() {}
diff --git a/Engine/Source/Developer/Horde/Public/Compute/AgentMessage.h b/Engine/Source/Developer/Horde/Public/Compute/AgentMessage.h
--- a/Engine/Source/Developer/Horde/Public/Compute/AgentMessage.h
+++ b/Engine/Source/Developer/Horde/Public/Compute/AgentMessage.h
@@ -134,10 +134,10 @@ public:
 	HORDE_API EAgentMessageType ReadResponse(int32 TimeoutMS = -1, bool* bOutTimedOut = nullptr);
 
 	// Gets the raw data from a response
-	HORDE_API const void* GetResponseData() const { return ResponseData; }
+	const void* GetResponseData() const { return ResponseData; }
 
 	// Gets the size of the response
-	HORDE_API size_t GetResponseSize() const { return ResponseLength; }
+	size_t GetResponseSize() const { return ResponseLength; }
 
 	// Reads an exception response
 	HORDE_API void ReadException(AgentMessage::FException& Ex);
diff --git a/Engine/Source/Developer/Horde/Public/Compute/ComputeBuffer.h b/Engine/Source/Developer/Horde/Public/Compute/ComputeBuffer.h
--- a/Engine/Source/Developer/Horde/Public/Compute/ComputeBuffer.h
+++ b/Engine/Source/Developer/Horde/Public/Compute/ComputeBuffer.h
@@ -61,7 +61,7 @@ public:
 	HORDE_API void Close();
 
 	// Test if the buffer is currently open
-	HORDE_API bool IsValid() const { return Detail != nullptr; }
+	bool IsValid() const { return Detail != nullptr; }
 
 	// Creates a new reader for this buffer
 	HORDE_API FComputeBufferReader CreateReader();
diff --git a/Engine/Source/Developer/Horde/Public/Compute/ComputeSocket.h b/Engine/Source/Developer/Horde/Public/Compute/ComputeSocket.h
--- a/Engine/Source/Developer/Horde/Public/Compute/ComputeSocket.h
+++ b/Engine/Source/Developer/Horde/Public/Compute/ComputeSocket.h
@@ -22,13 +22,13 @@ public:
 	FComputeSocket& operator=(const FComputeSocket&) = delete;
 
 	// Begins communication with the agent
-	HORDE_API virtual void StartCommunication() = 0;
+	virtual void StartCommunication() = 0;
 
 	// Attaches a new buffer for receiving data
-	HORDE_API virtual void AttachRecvBuffer(int ChannelId, FComputeBuffer RecvBuffer) = 0;
+	virtual void AttachRecvBuffer(int ChannelId, FComputeBuffer RecvBuffer) = 0;
 
 	// Attaches a new buffer for sending data */
-	HORDE_API virtual void AttachSendBuffer(int ChannelId, FComputeBuffer SendBuffer) = 0;
+	virtual void AttachSendBuffer(int ChannelId, FComputeBuffer SendBuffer) = 0;
 
 	// Attaches a channel to this socket
 	HORDE_API TSharedPtr<FComputeChannel> CreateChannel(int ChannelId, bool Anonymous = false);
diff --git a/Engine/Source/Developer/Horde/Public/HordeHttpClient.h b/Engine/Source/Developer/Horde/Public/HordeHttpClient.h
--- a/Engine/Source/Developer/Horde/Public/HordeHttpClient.h
+++ b/Engine/Source/Developer/Horde/Public/HordeHttpClient.h
@@ -7,20 +7,22 @@
 #include "Interfaces/IHttpRequest.h"
 #include "Interfaces/IHttpResponse.h"
 
-class HORDE_API FHordeHttpClient
+#define UE_API HORDE_API
+
+class FHordeHttpClient
 {
 public:
-	FHordeHttpClient(FString InServerUrl);
-	~FHordeHttpClient();
+	UE_API FHordeHttpClient(FString InServerUrl);
+	UE_API ~FHordeHttpClient();
 
-	bool Login(bool bUnattended, FFeedbackContext* Warn = nullptr);
+	UE_API bool Login(bool bUnattended, FFeedbackContext* Warn = nullptr);
 
-	bool LoginWithOidc(const TCHAR* Profile, bool bUnattended, FFeedbackContext* Warn = nullptr);
+	UE_API bool LoginWithOidc(const TCHAR* Profile, bool bUnattended, FFeedbackContext* Warn = nullptr);
 
-	bool LoginWithEnvironmentVariable();
+	UE_API bool LoginWithEnvironmentVariable();
 	
-	TSharedRef<IHttpRequest> CreateRequest(const TCHAR* Verb, const TCHAR* Path);
-	TSharedRef<IHttpResponse> Get(const TCHAR* Path);
+	UE_API TSharedRef<IHttpRequest> CreateRequest(const TCHAR* Verb, const TCHAR* Path);
+	UE_API TSharedRef<IHttpResponse> Get(const TCHAR* Path);
 
 	template<typename T>
 	T Get(const TCHAR* Path)
@@ -30,9 +32,11 @@ public:
 		return Result;
 	}
 
-	static TSharedRef<IHttpResponse> ExecuteRequest(TSharedRef<IHttpRequest> Request);
+	static UE_API TSharedRef<IHttpResponse> ExecuteRequest(TSharedRef<IHttpRequest> Request);
 
 private:
 	FString ServerUrl;
 	FString Token;
 };
+
+#undef UE_API
diff --git a/Engine/Source/Developer/Horde/Public/SharedBufferView.h b/Engine/Source/Developer/Horde/Public/SharedBufferView.h
--- a/Engine/Source/Developer/Horde/Public/SharedBufferView.h
+++ b/Engine/Source/Developer/Horde/Public/SharedBufferView.h
@@ -4,28 +4,32 @@
 
 #include "Memory/SharedBuffer.h"
 
+#define UE_API HORDE_API
+
 //
 // View into a buffer with shared ownership
 //
-class HORDE_API FSharedBufferView final
+class FSharedBufferView final
 {
 public:
-	FSharedBufferView();
-	FSharedBufferView(FSharedBuffer InBuffer);
-	FSharedBufferView(FSharedBuffer InBuffer, const FMemoryView& InView);
-	FSharedBufferView(FSharedBuffer InBuffer, size_t InOffset, size_t InLength);
-	~FSharedBufferView();
+	UE_API FSharedBufferView();
+	UE_API FSharedBufferView(FSharedBuffer InBuffer);
+	UE_API FSharedBufferView(FSharedBuffer InBuffer, const FMemoryView& InView);
+	UE_API FSharedBufferView(FSharedBuffer InBuffer, size_t InOffset, size_t InLength);
+	UE_API ~FSharedBufferView();
 
-	static FSharedBufferView Copy(const FMemoryView& View);
+	static UE_API FSharedBufferView Copy(const FMemoryView& View);
 
-	FSharedBufferView Slice(uint64 Offset) const;
-	FSharedBufferView Slice(uint64 Offset, uint64 Length) const;
+	UE_API FSharedBufferView Slice(uint64 Offset) const;
+	UE_API FSharedBufferView Slice(uint64 Offset, uint64 Length) const;
 
-	const unsigned char* GetPointer() const;
-	size_t GetLength() const;
-	FMemoryView GetView() const;
+	UE_API const unsigned char* GetPointer() const;
+	UE_API size_t GetLength() const;
+	UE_API FMemoryView GetView() const;
 
 private:
 	FSharedBuffer Buffer;
 	FMemoryView View;
 };
+
+#undef UE_API
diff --git a/Engine/Source/Developer/Horde/Public/Storage/Blob.h b/Engine/Source/Developer/Horde/Public/Storage/Blob.h
--- a/Engine/Source/Developer/Horde/Public/Storage/Blob.h
+++ b/Engine/Source/Developer/Horde/Public/Storage/Blob.h
@@ -7,10 +7,12 @@
 #include "BlobType.h"
 #include "../SharedBufferView.h"
 
+#define UE_API HORDE_API
+
 /**
  * Describes a blob of data in the storage system
  */
-struct HORDE_API FBlob
+struct FBlob
 {
 	/** Type of the blob. */
 	FBlobType Type;
@@ -21,6 +23,8 @@ struct HORDE_API FBlob
 	/** References to other blobs. */
 	TArray<FBlobHandle> References;
 
-	FBlob(const FBlobType& InType, FSharedBufferView InData, TArray<FBlobHandle> InReferences);
-	~FBlob();
+	UE_API FBlob(const FBlobType& InType, FSharedBufferView InData, TArray<FBlobHandle> InReferences);
+	UE_API ~FBlob();
 };
+
+#undef UE_API
diff --git a/Engine/Source/Developer/Horde/Public/Storage/BlobHandle.h b/Engine/Source/Developer/Horde/Public/Storage/BlobHandle.h
--- a/Engine/Source/Developer/Horde/Public/Storage/BlobHandle.h
+++ b/Engine/Source/Developer/Horde/Public/Storage/BlobHandle.h
@@ -7,6 +7,8 @@
 #include "BlobType.h"
 #include "IO/IoHash.h"
 
+#define UE_API HORDE_API
+
 struct FBlob;
 class FBlobHandleData;
 class FBlobLocator;
@@ -14,7 +16,7 @@ class FBlobLocator;
 /**
  * Handle to a blob. Can be used to reference blobs that have not been flushed yet. Implemented as a shared pointer to a FBlobHandleData object.
  */
-class HORDE_API FBlobHandle : public TSharedPtr<FBlobHandleData>
+class FBlobHandle : public TSharedPtr<FBlobHandleData>
 {
 public:
 	using Super = TSharedPtr<FBlobHandleData>;
@@ -43,8 +45,8 @@ public:
 		: TSharedPtr<FBlobHandleData>(StaticCastSharedPtr<FBlobHandleData, T>(Value.ToSharedPtr()))
 	{ }
 
-	bool operator==(const FBlobHandle& Other) const;
-	bool operator!=(const FBlobHandle& Other) const;
+	UE_API bool operator==(const FBlobHandle& Other) const;
+	UE_API bool operator!=(const FBlobHandle& Other) const;
 
 	friend uint32 GetTypeHash(const FBlobHandle& Handle);
 };
@@ -70,10 +72,10 @@ public:
 /**
  * Base interface for a blob.
  */
-class HORDE_API FBlobHandleData
+class FBlobHandleData
 {
 public:
-	virtual ~FBlobHandleData();
+	UE_API virtual ~FBlobHandleData();
 
 	/** Gets the type name of this derived blob handle instance. */
 	virtual const char* GetType() const = 0;
@@ -91,31 +93,31 @@ public:
 	virtual FBlob Read() const = 0;
 
 	/** Gets the type of this blob. */
-	virtual FBlobType ReadType() const;
+	UE_API virtual FBlobType ReadType() const;
 
 	/** Gets the outward references from this blob. */
-	virtual void ReadImports(TArray<FBlobHandle>& OutImports) const;
+	UE_API virtual void ReadImports(TArray<FBlobHandle>& OutImports) const;
 
 	/** Reads part of the blob, and returns a handle that can be used to access the data. */
-	virtual FSharedBufferView ReadBody() const;
+	UE_API virtual FSharedBufferView ReadBody() const;
 
 	/** Reads part of the blob, and returns a handle that can be used to access the data. */
-	virtual FSharedBufferView ReadBody(size_t Offset, TOptional<size_t> Length = TOptional<size_t>()) const;
+	UE_API virtual FSharedBufferView ReadBody(size_t Offset, TOptional<size_t> Length = TOptional<size_t>()) const;
 
 	/** Gets a locator for this blob, by calling TryAppendLocator recursively. */
-	FBlobLocator GetLocator() const;
+	UE_API FBlobLocator GetLocator() const;
 
 	/** Gets an identifier for this blob, relative to its outer. */
 	virtual bool TryAppendIdentifier(FUtf8String& OutBuffer) const = 0;
 
 	/** Gets a handle to a nested blob object. */
-	virtual FBlobHandle GetFragmentHandle(const FUtf8StringView& Fragment) const;
+	UE_API virtual FBlobHandle GetFragmentHandle(const FUtf8StringView& Fragment) const;
 };
 
 /**
  * Stores a blob handle along with a hash of the target node.
  */
-struct HORDE_API FBlobHandleWithHash
+struct FBlobHandleWithHash
 {
 	FBlobHandle Handle;
 	FIoHash Hash;
@@ -125,3 +127,5 @@ struct HORDE_API FBlobHandleWithHash
 		, Hash(InHash)
 	{ }
 };
+
+#undef UE_API
diff --git a/Engine/Source/Developer/Horde/Public/Storage/BlobLocator.h b/Engine/Source/Developer/Horde/Public/Storage/BlobLocator.h
--- a/Engine/Source/Developer/Horde/Public/Storage/BlobLocator.h
+++ b/Engine/Source/Developer/Horde/Public/Storage/BlobLocator.h
@@ -5,40 +5,44 @@
 #include "CoreMinimal.h"
 #include "Containers/Utf8String.h"
 
+#define UE_API HORDE_API
+
 /**
  * Identifies the location of a blob. Meaning of this string is implementation defined.
  */
-class HORDE_API FBlobLocator
+class FBlobLocator
 {
 public:
-	FBlobLocator();
-	FBlobLocator(FUtf8String InPath);
-	FBlobLocator(const FBlobLocator& InBaseLocator, const FUtf8StringView& InFragment);
-	explicit FBlobLocator(const FUtf8StringView& InPath);
+	UE_API FBlobLocator();
+	UE_API FBlobLocator(FUtf8String InPath);
+	UE_API FBlobLocator(const FBlobLocator& InBaseLocator, const FUtf8StringView& InFragment);
+	UE_API explicit FBlobLocator(const FUtf8StringView& InPath);
 
 	/** Tests whether this locator is valid. */
-	bool IsValid() const;
+	UE_API bool IsValid() const;
 
 	/** Gets the path for this locator. */
-	const FUtf8String& GetPath() const;
+	UE_API const FUtf8String& GetPath() const;
 
 	/** Gets the base locator for this blob. */
-	FBlobLocator GetBaseLocator() const;
+	UE_API FBlobLocator GetBaseLocator() const;
 
 	/** Gets the path for this locator. */
-	FUtf8StringView GetFragment() const;
+	UE_API FUtf8StringView GetFragment() const;
 
 	/** Determines if this locator can be unwrapped into an outer locator/fragment pair. */
-	bool CanUnwrap() const;
+	UE_API bool CanUnwrap() const;
 
 	/** Split this locator into a locator and fragment. */
-	bool TryUnwrap(FBlobLocator& OutLocator, FUtf8StringView& OutFragment) const;
+	UE_API bool TryUnwrap(FBlobLocator& OutLocator, FUtf8StringView& OutFragment) const;
 
-	bool operator==(const FBlobLocator& Other) const;
-	bool operator!=(const FBlobLocator& Other) const;
+	UE_API bool operator==(const FBlobLocator& Other) const;
+	UE_API bool operator!=(const FBlobLocator& Other) const;
 
 	friend uint32 GetTypeHash(const FBlobLocator& Locator);
 
 private:
 	FUtf8String Path;
-};
\ No newline at end of file
+};
+
+#undef UE_API
diff --git a/Engine/Source/Developer/Horde/Public/Storage/BlobReader.h b/Engine/Source/Developer/Horde/Public/Storage/BlobReader.h
--- a/Engine/Source/Developer/Horde/Public/Storage/BlobReader.h
+++ b/Engine/Source/Developer/Horde/Public/Storage/BlobReader.h
@@ -8,31 +8,33 @@
 #include "BlobLocator.h"
 #include "RefName.h"
 
+#define UE_API HORDE_API
+
 struct FIoHash;
 
 /**
  * Reads data from a blob object
  */
-class HORDE_API FBlobReader
+class FBlobReader
 {
 public:
-	FBlobReader(const FBlob& InBlob);
-	FBlobReader(int32 InVersion, const FMemoryView& InBuffer, const TArray<FBlobHandle>& InImports);
+	UE_API FBlobReader(const FBlob& InBlob);
+	UE_API FBlobReader(int32 InVersion, const FMemoryView& InBuffer, const TArray<FBlobHandle>& InImports);
 
 	/** Gets a version number for the current blob */
-	int GetVersion() const;
+	UE_API int GetVersion() const;
 
 	/** Gets a pointer to the remaining memory. */
-	const unsigned char* GetBuffer() const;
+	UE_API const unsigned char* GetBuffer() const;
 
 	/** Gets the remaining memory to read from. */
-	FMemoryView GetView() const;
+	UE_API FMemoryView GetView() const;
 
 	/** Advance the current read position. */
-	void Advance(size_t Size);
+	UE_API void Advance(size_t Size);
 
 	/** Reads the next import. */
-	FBlobHandle ReadImport();
+	UE_API FBlobHandle ReadImport();
 
 private:
 	int32 Version;
@@ -51,3 +53,5 @@ HORDE_API FMemoryView ReadFixedLengthBytes(FBlobReader& Reader, size_t Length);
 HORDE_API size_t ReadUnsignedVarInt(FBlobReader& Reader);
 HORDE_API FUtf8String ReadString(FBlobReader& Reader);
 HORDE_API FMemoryView ReadStringSpan(FBlobReader& Reader);
+
+#undef UE_API
diff --git a/Engine/Source/Developer/Horde/Public/Storage/BlobType.h b/Engine/Source/Developer/Horde/Public/Storage/BlobType.h
--- a/Engine/Source/Developer/Horde/Public/Storage/BlobType.h
+++ b/Engine/Source/Developer/Horde/Public/Storage/BlobType.h
@@ -4,12 +4,14 @@
 
 #include "CoreMinimal.h"
 
+#define UE_API HORDE_API
+
 /**
  * Identifies the type of a blob
  */
-struct HORDE_API FBlobType
+struct FBlobType
 {
-	static const FBlobType Leaf;
+	static UE_API const FBlobType Leaf;
 
 	/** Nominal identifier for the type. */
 	FGuid Guid;
@@ -17,8 +19,10 @@ struct HORDE_API FBlobType
 	/** Version number for the serializer. */
 	int Version;
 
-	FBlobType(const FGuid& InGuid, int InVersion);
+	UE_API FBlobType(const FGuid& InGuid, int InVersion);
 
-	bool operator==(const FBlobType& Other) const;
-	bool operator!=(const FBlobType& Other) const;
+	UE_API bool operator==(const FBlobType& Other) const;
+	UE_API bool operator!=(const FBlobType& Other) const;
 };
+
+#undef UE_API
diff --git a/Engine/Source/Developer/Horde/Public/Storage/BlobWriter.h b/Engine/Source/Developer/Horde/Public/Storage/BlobWriter.h
--- a/Engine/Source/Developer/Horde/Public/Storage/BlobWriter.h
+++ b/Engine/Source/Developer/Horde/Public/Storage/BlobWriter.h
@@ -9,16 +9,18 @@
 #include "RefName.h"
 #include "../SharedBufferView.h"
 
+#define UE_API HORDE_API
+
 struct FAliasInfo;
 struct FRefOptions;
 
 /**
  * Interface for writing data to the storage system.
  */
-class HORDE_API FBlobWriter
+class FBlobWriter
 {
 public:
-	virtual ~FBlobWriter();
+	UE_API virtual ~FBlobWriter();
 
 	/** Create another writer instance, allowing multiple threads to write in parallel. */
 	virtual TUniquePtr<FBlobWriter> Fork() = 0;
@@ -30,11 +32,11 @@ public:
 	virtual void AddImport(FBlobHandle Target) = 0;
 
 	/** Adds a named reference to the blob being built. */
-	void AddRef(const FRefName& RefName);
+	UE_API void AddRef(const FRefName& RefName);
 	virtual void AddRef(const FRefName& RefName, const FRefOptions& Options) = 0;
 
 	/** Gets a block of memory of the given size. */
-	void* GetOutputBuffer(size_t Size);
+	UE_API void* GetOutputBuffer(size_t Size);
 
 	/** Gets a block of memory, at least the given size. */
 	virtual FMutableMemoryView GetOutputBufferAsSpan(size_t UsedSize, size_t DesiredSize) = 0;
@@ -69,3 +71,5 @@ HORDE_API size_t MeasureString(const FUtf8StringView& Text);
 HORDE_API void WriteString(FBlobWriter& Writer, const char* Text);
 HORDE_API void WriteString(FBlobWriter& Writer, const FUtf8StringView& Text);
 HORDE_API void WriteString(FBlobWriter& Writer, const FUtf8String& Text);
+
+#undef UE_API
diff --git a/Engine/Source/Developer/Horde/Public/Storage/Bundles/Bundle.h b/Engine/Source/Developer/Horde/Public/Storage/Bundles/Bundle.h
--- a/Engine/Source/Developer/Horde/Public/Storage/Bundles/Bundle.h
+++ b/Engine/Source/Developer/Horde/Public/Storage/Bundles/Bundle.h
@@ -4,6 +4,8 @@
 
 #include "Storage/Blob.h"
 
+#define UE_API HORDE_API
+
 /** Bundle version number */
 enum class EBundleVersion : unsigned char
 {
@@ -41,15 +43,15 @@ enum class EBundleVersion : unsigned char
 /**
  * General bundle properties
  */
-struct HORDE_API FBundle
+struct FBundle
 {
-	static const FBlobType BlobType;
+	static UE_API const FBlobType BlobType;
 };
 
 /**
  * Signature for a bundle
  */
-struct HORDE_API FBundleSignature
+struct FBundleSignature
 {
 	/** Number of bytes in a serialized signature. */
 	static const int NumBytes = 8;
@@ -60,7 +62,9 @@ struct HORDE_API FBundleSignature
 	/** Length of the packet, including the length of the signature itself. */
 	size_t Length;
 
-	FBundleSignature(EBundleVersion InVersion, size_t InLength);
-	static FBundleSignature Read(const void* Data);
-	void Write(void* Data);
+	UE_API FBundleSignature(EBundleVersion InVersion, size_t InLength);
+	static UE_API FBundleSignature Read(const void* Data);
+	UE_API void Write(void* Data);
 };
+
+#undef UE_API
diff --git a/Engine/Source/Developer/Horde/Public/Storage/Bundles/BundleCompression.h b/Engine/Source/Developer/Horde/Public/Storage/Bundles/BundleCompression.h
--- a/Engine/Source/Developer/Horde/Public/Storage/Bundles/BundleCompression.h
+++ b/Engine/Source/Developer/Horde/Public/Storage/Bundles/BundleCompression.h
@@ -5,6 +5,8 @@
 #include "CoreMinimal.h"
 #include "Memory/MemoryView.h"
 
+#define UE_API HORDE_API
+
 /** Indicates the compression format in a bundle */
 enum EBundleCompressionFormat : unsigned char
 {
@@ -27,14 +29,16 @@ enum EBundleCompressionFormat : unsigned char
 /*
  * Utility methods for compressing bundles
  */
-struct HORDE_API FBundleCompression
+struct FBundleCompression
 {
 	/** Gets the maximum size of the buffer required to compress the given data. */
-	static size_t GetMaxSize(EBundleCompressionFormat Format, const FMemoryView& Input);
+	static UE_API size_t GetMaxSize(EBundleCompressionFormat Format, const FMemoryView& Input);
 
 	/** Compress a data packet. */
-	static size_t Compress(EBundleCompressionFormat Format, const FMemoryView& Input, const FMutableMemoryView& Output);
+	static UE_API size_t Compress(EBundleCompressionFormat Format, const FMemoryView& Input, const FMutableMemoryView& Output);
 
 	/** Decompress a packet of data. */
-	static void Decompress(EBundleCompressionFormat Format, const FMemoryView& Input, const FMutableMemoryView& Output);
+	static UE_API void Decompress(EBundleCompressionFormat Format, const FMemoryView& Input, const FMutableMemoryView& Output);
 };
+
+#undef UE_API
diff --git a/Engine/Source/Developer/Horde/Public/Storage/Bundles/BundleOptions.h b/Engine/Source/Developer/Horde/Public/Storage/Bundles/BundleOptions.h
--- a/Engine/Source/Developer/Horde/Public/Storage/Bundles/BundleOptions.h
+++ b/Engine/Source/Developer/Horde/Public/Storage/Bundles/BundleOptions.h
@@ -5,13 +5,15 @@
 #include "Bundle.h"
 #include "BundleCompression.h"
 
+#define UE_API HORDE_API
+
 /*
  * Options for configuring a bundle serializer
  */
-struct HORDE_API FBundleOptions
+struct FBundleOptions
 {
 	/** Default options value. */
-	static const FBundleOptions Default;
+	static UE_API const FBundleOptions Default;
 
 	/** Maximum version number of bundles to write. */
 	EBundleVersion MaxVersion = EBundleVersion::LatestV2;
@@ -25,6 +27,7 @@ struct HORDE_API FBundleOptions
 	/** Minimum size of a block to be compressed. */
 	int MinCompressionPacketSize = 16 * 1024;
 
-	FBundleOptions();
+	UE_API FBundleOptions();
 };
 
+#undef UE_API
diff --git a/Engine/Source/Developer/Horde/Public/Storage/Bundles/V2/BundleWriter.h b/Engine/Source/Developer/Horde/Public/Storage/Bundles/V2/BundleWriter.h
--- a/Engine/Source/Developer/Horde/Public/Storage/Bundles/V2/BundleWriter.h
+++ b/Engine/Source/Developer/Horde/Public/Storage/Bundles/V2/BundleWriter.h
@@ -10,26 +10,28 @@
 #include "Storage/Clients/BundleStorageClient.h"
 #include "Storage/Clients/KeyValueStorageClient.h"
 
+#define UE_API HORDE_API
+
 class FPacketWriter;
 
 /**
  * Writes blobs into bundles
  */
-class HORDE_API FBundleWriter final : public FBlobWriter
+class FBundleWriter final : public FBlobWriter
 {
 public:
-	FBundleWriter(TSharedRef<FKeyValueStorageClient> InStorageClient, FUtf8String InBasePath, const FBundleOptions& InOptions);
-	~FBundleWriter();
+	UE_API FBundleWriter(TSharedRef<FKeyValueStorageClient> InStorageClient, FUtf8String InBasePath, const FBundleOptions& InOptions);
+	UE_API ~FBundleWriter();
 
 	// Inherited from FBlobWriter
-	void Flush() override;
-	TUniquePtr<FBlobWriter> Fork() override;
-	virtual void AddAlias(const FAliasInfo& AliasInfo) override;
-	virtual void AddImport(FBlobHandle Target) override;
-	virtual void AddRef(const FRefName& RefName, const FRefOptions& Options) override;
-	virtual FBlobHandle CompleteBlob(const FBlobType& InType) override;
-	virtual FMutableMemoryView GetOutputBufferAsSpan(size_t UsedSize, size_t DesiredSize) override;
-	virtual void Advance(size_t Size) override;
+	UE_API void Flush() override;
+	UE_API TUniquePtr<FBlobWriter> Fork() override;
+	UE_API virtual void AddAlias(const FAliasInfo& AliasInfo) override;
+	UE_API virtual void AddImport(FBlobHandle Target) override;
+	UE_API virtual void AddRef(const FRefName& RefName, const FRefOptions& Options) override;
+	UE_API virtual FBlobHandle CompleteBlob(const FBlobType& InType) override;
+	UE_API virtual FMutableMemoryView GetOutputBufferAsSpan(size_t UsedSize, size_t DesiredSize) override;
+	UE_API virtual void Advance(size_t Size) override;
 
 private:
 	class FPendingExportHandleData;
@@ -51,6 +53,8 @@ private:
 	//	std::vector<std::pair<FBlobHandle, FAliasInfo>> PendingExportAliases;
 	//	FChunkedBlobWriter EncodedPacketWriter;
 
-	void StartPacket();
-	void FinishPacket();
-};
\ No newline at end of file
+	UE_API void StartPacket();
+	UE_API void FinishPacket();
+};
+
+#undef UE_API
diff --git a/Engine/Source/Developer/Horde/Public/Storage/Bundles/V2/ExportHandle.h b/Engine/Source/Developer/Horde/Public/Storage/Bundles/V2/ExportHandle.h
--- a/Engine/Source/Developer/Horde/Public/Storage/Bundles/V2/ExportHandle.h
+++ b/Engine/Source/Developer/Horde/Public/Storage/Bundles/V2/ExportHandle.h
@@ -6,6 +6,8 @@
 #include "Storage/BlobHandle.h"
 #include "PacketHandle.h"
 
+#define UE_API HORDE_API
+
 class FExportHandleData;
 
 /** Handle to an export within a packet. */
@@ -14,30 +16,32 @@ typedef TBlobHandle<FExportHandleData> FExportHandle;
 /** 
  * Implementation of export handle data.
  */
-class HORDE_API FExportHandleData final : public FBlobHandleData, public TSharedFromThis<FExportHandleData, ESPMode::ThreadSafe>
+class FExportHandleData final : public FBlobHandleData, public TSharedFromThis<FExportHandleData, ESPMode::ThreadSafe>
 {
 public:
-	static const char Type[];
+	static UE_API const char Type[];
 
-	FExportHandleData(FPacketHandle InPacket, int32 InExportIdx);
-	FExportHandleData(FPacketHandle InPacket, const FUtf8StringView& InFragment);
-	~FExportHandleData();
+	UE_API FExportHandleData(FPacketHandle InPacket, int32 InExportIdx);
+	UE_API FExportHandleData(FPacketHandle InPacket, const FUtf8StringView& InFragment);
+	UE_API ~FExportHandleData();
 
-	static void AppendIdentifier(FUtf8String& OutBuffer, int32 ExportIdx);
+	static UE_API void AppendIdentifier(FUtf8String& OutBuffer, int32 ExportIdx);
 
 	// Implementation of FBlobHandle
-	virtual bool Equals(const FBlobHandleData& Other) const override;
-	virtual uint32 GetHashCode() const override;
-	virtual FBlobHandle GetOuter() const override;
-	virtual const char* GetType() const override;
-	virtual FBlob Read() const override;
-	virtual bool TryAppendIdentifier(FUtf8String& OutBuffer) const override;
+	UE_API virtual bool Equals(const FBlobHandleData& Other) const override;
+	UE_API virtual uint32 GetHashCode() const override;
+	UE_API virtual FBlobHandle GetOuter() const override;
+	UE_API virtual const char* GetType() const override;
+	UE_API virtual FBlob Read() const override;
+	UE_API virtual bool TryAppendIdentifier(FUtf8String& OutBuffer) const override;
 
 private:
-	static const FUtf8StringView FragmentPrefix;
+	static UE_API const FUtf8StringView FragmentPrefix;
 
 	FPacketHandle Packet;
 	int32 ExportIdx;
 
-	static bool TryParse(const FUtf8StringView& Fragment, int32& OutExportIdx);
+	static UE_API bool TryParse(const FUtf8StringView& Fragment, int32& OutExportIdx);
 };
+
+#undef UE_API
diff --git a/Engine/Source/Developer/Horde/Public/Storage/Bundles/V2/Packet.h b/Engine/Source/Developer/Horde/Public/Storage/Bundles/V2/Packet.h
--- a/Engine/Source/Developer/Horde/Public/Storage/Bundles/V2/Packet.h
+++ b/Engine/Source/Developer/Horde/Public/Storage/Bundles/V2/Packet.h
@@ -8,6 +8,8 @@
 #include "Storage/BlobType.h"
 #include "Storage/Bundles/BundleCompression.h"
 
+#define UE_API HORDE_API
+
 class FChunkedBufferWriter;
 struct FPacketImport;
 struct FPacketExport;
@@ -15,43 +17,43 @@ struct FPacketExport;
 /**
  * Accessor for data structures stored into a serialized bundle packet.
  */
-class HORDE_API FPacket
+class FPacket
 {
 public:
-	static const FBlobType BlobType;
+	static UE_API const FBlobType BlobType;
 
-	FPacket(FSharedBufferView InBuffer);
-	virtual ~FPacket();
+	UE_API FPacket(FSharedBufferView InBuffer);
+	UE_API virtual ~FPacket();
 
 	/** Gets the underlying buffer for this packet. */
-	FSharedBufferView GetBuffer() const;
+	UE_API FSharedBufferView GetBuffer() const;
 
 	/** Length of this packet. */
-	size_t GetLength() const;
+	UE_API size_t GetLength() const;
 
 	/** Gets the number of types in this packet. */
-	size_t GetTypeCount() const;
+	UE_API size_t GetTypeCount() const;
 
 	/** Gets a type from the packet. */
-	const FBlobType& GetType(size_t TypeIdx) const;
+	UE_API const FBlobType& GetType(size_t TypeIdx) const;
 
 	/** Gets the number of imports in this packet. */
-	size_t GetImportCount() const;
+	UE_API size_t GetImportCount() const;
 
 	/** Gets the locator for a particular import. */
-	FPacketImport GetImport(size_t ImportIdx) const;
+	UE_API FPacketImport GetImport(size_t ImportIdx) const;
 
 	/** Gets the number of exports in this packet. */
-	size_t GetExportCount() const;
+	UE_API size_t GetExportCount() const;
 
 	/** Gets the bulk data for a particular export. */
-	FPacketExport GetExport(size_t ExportIdx) const;
+	UE_API FPacketExport GetExport(size_t ExportIdx) const;
 
 	/** Encodes a packet. */
-	void Encode(EBundleCompressionFormat Format, FChunkedBufferWriter& Writer);
+	UE_API void Encode(EBundleCompressionFormat Format, FChunkedBufferWriter& Writer);
 
 	/** Decodes a packet from the given data. */
-	static FPacket Decode(const FMemoryView& View);
+	static UE_API FPacket Decode(const FMemoryView& View);
 
 private:
 	struct FEncodedPacketHeader;
@@ -62,21 +64,21 @@ private:
 /**
  * Specifies the path to an imported node
  */
-struct HORDE_API FPacketImport
+struct FPacketImport
 {
 	static constexpr int32 Bias = 3;
 	static constexpr int32 InvalidBaseIdx = -1;
 	static constexpr int32 CurrentPacketBaseIdx = -2;
 	static constexpr int32 CurrentBundleBaseIdx = -3;
 
-	FPacketImport(const FMemoryView& InView);
-	FPacketImport(int32 InBaseIdx, const FUtf8StringView& InFragment);
+	UE_API FPacketImport(const FMemoryView& InView);
+	UE_API FPacketImport(int32 InBaseIdx, const FUtf8StringView& InFragment);
 
 	/** Gets the index of the base packet import. */
-	int GetBaseIdx() const;
+	UE_API int GetBaseIdx() const;
 
 	/** Gets the fragment for this import. */
-	FUtf8StringView GetFragment() const;
+	UE_API FUtf8StringView GetFragment() const;
 
 private:
 	int32 BaseIdx;
@@ -86,22 +88,24 @@ private:
 /*
  * Data for an exported node in a packet
  */
-struct HORDE_API FPacketExport
+struct FPacketExport
 {
-	FPacketExport(FSharedBufferView InBuffer);
-	~FPacketExport();
+	UE_API FPacketExport(FSharedBufferView InBuffer);
+	UE_API ~FPacketExport();
 
 	/** Gets the index of this export's type. */
-	size_t GetTypeIndex() const;
+	UE_API size_t GetTypeIndex() const;
 
 	/** Gets the import indexes for this export. */
-	void GetImportIndices(TArray<size_t>& OutIndices) const;
+	UE_API void GetImportIndices(TArray<size_t>& OutIndices) const;
 
 	/** Gets the payload for this export. */
-	FSharedBufferView GetPayload() const;
+	UE_API FSharedBufferView GetPayload() const;
 
 private:
 	FSharedBufferView Buffer;
 	int32 TypeIdx;
 	FMemoryView Imports;
 };
+
+#undef UE_API
diff --git a/Engine/Source/Developer/Horde/Public/Storage/Bundles/V2/PacketHandle.h b/Engine/Source/Developer/Horde/Public/Storage/Bundles/V2/PacketHandle.h
--- a/Engine/Source/Developer/Horde/Public/Storage/Bundles/V2/PacketHandle.h
+++ b/Engine/Source/Developer/Horde/Public/Storage/Bundles/V2/PacketHandle.h
@@ -9,6 +9,8 @@
 #include "Storage/BlobType.h"
 #include "Storage/StorageClient.h"
 
+#define UE_API HORDE_API
+
 class FPacketReader;
 class FPacketHandleData;
 
@@ -16,33 +18,33 @@ class FPacketHandleData;
 typedef TBlobHandle<FPacketHandleData> FPacketHandle;
 
 // Data for FPacketHandle
-class HORDE_API FPacketHandleData final : public FBlobHandleData, public TSharedFromThis<FPacketHandleData, ESPMode::ThreadSafe>
+class FPacketHandleData final : public FBlobHandleData, public TSharedFromThis<FPacketHandleData, ESPMode::ThreadSafe>
 {
 public:
-	static const char Type[];
+	static UE_API const char Type[];
 
-	FPacketHandleData(TSharedRef<FStorageClient> InStorageClient, FBlobHandle InBundle, size_t InPacketOffset, size_t InPacketLength);
-	FPacketHandleData(TSharedRef<FStorageClient> InStorageClient, FBlobHandle InBundle, const FUtf8StringView& InFragment);
-	~FPacketHandleData();
+	UE_API FPacketHandleData(TSharedRef<FStorageClient> InStorageClient, FBlobHandle InBundle, size_t InPacketOffset, size_t InPacketLength);
+	UE_API FPacketHandleData(TSharedRef<FStorageClient> InStorageClient, FBlobHandle InBundle, const FUtf8StringView& InFragment);
+	UE_API ~FPacketHandleData();
 
 	/** Reads an export from this packet. */
-	FBlob ReadExport(size_t ExportIdx) const;
+	UE_API FBlob ReadExport(size_t ExportIdx) const;
 
 	/** Reads an export body from this packet. */
-	FSharedBufferView ReadExportBody(size_t ExportIdx) const;
+	UE_API FSharedBufferView ReadExportBody(size_t ExportIdx) const;
 
 	// Implementation of FBlobHandle
-	virtual bool Equals(const FBlobHandleData& Other) const override;
-	virtual uint32 GetHashCode() const override;
-	virtual FBlobHandle GetOuter() const override;
-	virtual const char* GetType() const override;
-	virtual FBlob Read() const override;
-	virtual bool TryAppendIdentifier(FUtf8String& OutBuffer) const override;
-	virtual FBlobHandle GetFragmentHandle(const FUtf8StringView& Fragment) const override;
+	UE_API virtual bool Equals(const FBlobHandleData& Other) const override;
+	UE_API virtual uint32 GetHashCode() const override;
+	UE_API virtual FBlobHandle GetOuter() const override;
+	UE_API virtual const char* GetType() const override;
+	UE_API virtual FBlob Read() const override;
+	UE_API virtual bool TryAppendIdentifier(FUtf8String& OutBuffer) const override;
+	UE_API virtual FBlobHandle GetFragmentHandle(const FUtf8StringView& Fragment) const override;
 
 private:
-	static const char FragmentPrefix[];
-	static const size_t FragmentPrefixLength;
+	static UE_API const char FragmentPrefix[];
+	static UE_API const size_t FragmentPrefixLength;
 
 	TSharedRef<FStorageClient> StorageClient;
 
@@ -52,6 +54,8 @@ private:
 
 	mutable TSharedPtr<FPacketReader> PacketReader;
 
-	const FPacketReader& GetPacketReader() const;
-	static bool TryParse(const FUtf8StringView& Fragment, size_t& OutPacketOffset, size_t& OutPacketLength);
+	UE_API const FPacketReader& GetPacketReader() const;
+	static UE_API bool TryParse(const FUtf8StringView& Fragment, size_t& OutPacketOffset, size_t& OutPacketLength);
 };
+
+#undef UE_API
diff --git a/Engine/Source/Developer/Horde/Public/Storage/Bundles/V2/PacketReader.h b/Engine/Source/Developer/Horde/Public/Storage/Bundles/V2/PacketReader.h
--- a/Engine/Source/Developer/Horde/Public/Storage/Bundles/V2/PacketReader.h
+++ b/Engine/Source/Developer/Horde/Public/Storage/Bundles/V2/PacketReader.h
@@ -7,23 +7,25 @@
 #include "Packet.h"
 #include "PacketHandle.h"
 
+#define UE_API HORDE_API
+
 /**
  * Utility class for constructing BlobData objects from a packet, caching any computed handles to other blobs.
  */
-class HORDE_API FPacketReader
+class FPacketReader
 {
 public:
-	FPacketReader(TSharedRef<FStorageClient> InStorageClient, FBlobHandle InBundleHandle, FPacketHandle InPacketHandle, FPacket InPacket);
-	~FPacketReader();
+	UE_API FPacketReader(TSharedRef<FStorageClient> InStorageClient, FBlobHandle InBundleHandle, FPacketHandle InPacketHandle, FPacket InPacket);
+	UE_API ~FPacketReader();
 
 	/** Read a blob object for the entire packet. */
-	FBlob ReadPacket() const;
+	UE_API FBlob ReadPacket() const;
 
 	/** Reads an export from this packet. */
-	FBlob ReadExport(size_t ExportIdx) const;
+	UE_API FBlob ReadExport(size_t ExportIdx) const;
 
 	/** Reads an export body from this packet. */
-	FSharedBufferView ReadExportBody(size_t ExportIdx) const;
+	UE_API FSharedBufferView ReadExportBody(size_t ExportIdx) const;
 
 private:
 	TSharedRef<FStorageClient> StorageClient;
@@ -32,5 +34,7 @@ private:
 	FPacket Packet;
 	mutable TArray<FBlobHandle> CachedImportHandles;
 
-	FBlobHandle GetImportHandle(size_t Index) const;
+	UE_API FBlobHandle GetImportHandle(size_t Index) const;
 };
+
+#undef UE_API
diff --git a/Engine/Source/Developer/Horde/Public/Storage/Bundles/V2/PacketWriter.h b/Engine/Source/Developer/Horde/Public/Storage/Bundles/V2/PacketWriter.h
--- a/Engine/Source/Developer/Horde/Public/Storage/Bundles/V2/PacketWriter.h
+++ b/Engine/Source/Developer/Horde/Public/Storage/Bundles/V2/PacketWriter.h
@@ -8,45 +8,47 @@
 #include "PacketHandle.h"
 #include "Storage/ChunkedBufferWriter.h"
 
+#define UE_API HORDE_API
+
 /**
  * Writes exports into a new bundle packet
  */
-class HORDE_API FPacketWriter
+class FPacketWriter
 {
 public:
-	FPacketWriter(FBlobHandle InBundleHandle, FBlobHandle InPacketHandle);
+	UE_API FPacketWriter(FBlobHandle InBundleHandle, FBlobHandle InPacketHandle);
 	FPacketWriter(const FPacketWriter&) = delete;
-	~FPacketWriter();
+	UE_API ~FPacketWriter();
 
 	/** Current length of the packet */
-	size_t GetLength() const;
+	UE_API size_t GetLength() const;
 
 	/** Adds an import to the current blob */
-	void AddImport(FBlobHandle Import);
+	UE_API void AddImport(FBlobHandle Import);
 
 	/** Gets the number of unique imports current added to this packet */
-	int GetImportCount() const;
+	UE_API int GetImportCount() const;
 
 	/** Gets a packet import by index */
-	FBlobHandle GetImport(int ImportIdx) const;
+	UE_API FBlobHandle GetImport(int ImportIdx) const;
 
 	/** Gets data to write new export */
-	FMutableMemoryView GetOutputBuffer(size_t UsedSize, size_t DesiredSize);
+	UE_API FMutableMemoryView GetOutputBuffer(size_t UsedSize, size_t DesiredSize);
 
 	/** Increase the length of the current blob */
-	void Advance(size_t Size);
+	UE_API void Advance(size_t Size);
 
 	/** Writes a new blob to this packet */
-	int CompleteBlob(const FBlobType& Type);
+	UE_API int CompleteBlob(const FBlobType& Type);
 
 	/** Gets the number of exports currently in this packet */
-	int GetExportCount() const;
+	UE_API int GetExportCount() const;
 
 	/** Reads data for a blob written to storage */
-	FBlob GetExport(int ExportIdx) const;
+	UE_API FBlob GetExport(int ExportIdx) const;
 
 	/** Mark the current packet as complete */
-	FPacket CompletePacket();
+	UE_API FPacket CompletePacket();
 
 	FPacketWriter& operator=(const FPacketWriter&) = delete;
 
@@ -68,15 +70,17 @@ private:
 	TMap<FBlobHandle, int> ImportMap;
 	TArray<uint32> ExportOffsets;
 
-	FMutableMemoryView GetOutputBufferInternal(size_t UsedSize, size_t DesiredSize);
-	uint8* GetOutputSpanAndAdvance(size_t Length);
+	UE_API FMutableMemoryView GetOutputBufferInternal(size_t UsedSize, size_t DesiredSize);
+	UE_API uint8* GetOutputSpanAndAdvance(size_t Length);
 
-	int32 FindOrAddType(FBlobType Type);
-	int32 FindOrAddImport(const FBlobHandle& Handle);
+	UE_API int32 FindOrAddType(FBlobType Type);
+	UE_API int32 FindOrAddImport(const FBlobHandle& Handle);
 
-	void WriteTypeTable();
-	void WriteImportTable();
-	void WriteExportTable();
+	UE_API void WriteTypeTable();
+	UE_API void WriteImportTable();
+	UE_API void WriteExportTable();
 
-	static size_t Align(size_t Offset);
+	static UE_API size_t Align(size_t Offset);
 };
+
+#undef UE_API
diff --git a/Engine/Source/Developer/Horde/Public/Storage/ChunkedBufferWriter.h b/Engine/Source/Developer/Horde/Public/Storage/ChunkedBufferWriter.h
--- a/Engine/Source/Developer/Horde/Public/Storage/ChunkedBufferWriter.h
+++ b/Engine/Source/Developer/Horde/Public/Storage/ChunkedBufferWriter.h
@@ -5,37 +5,39 @@
 #include "CoreMinimal.h"
 #include "BlobWriter.h"
 
+#define UE_API HORDE_API
+
 /**
  * Writes data to a series of memory blocks
  */
-class HORDE_API FChunkedBufferWriter
+class FChunkedBufferWriter
 {
 public:
-	FChunkedBufferWriter(size_t InitialSize = 1024);
+	UE_API FChunkedBufferWriter(size_t InitialSize = 1024);
 	FChunkedBufferWriter(const FChunkedBufferWriter&) = delete;
 	FChunkedBufferWriter& operator = (const FChunkedBufferWriter&) = delete;
-	virtual ~FChunkedBufferWriter();
+	UE_API virtual ~FChunkedBufferWriter();
 
 	/** Reset the contents of this writer. */
-	void Reset();
+	UE_API void Reset();
 
 	/** Gets the current written length of this buffer. */
-	size_t GetLength() const;
+	UE_API size_t GetLength() const;
 
 	/** Get a handle to part of the written buffer */
-	FSharedBufferView Slice(size_t Offset, size_t Length) const;
+	UE_API FSharedBufferView Slice(size_t Offset, size_t Length) const;
 
 	/** Gets a view over the underlying buffer. */
-	TArray<FMemoryView> GetView() const;
+	UE_API TArray<FMemoryView> GetView() const;
 
 	/** Copies the entire contents of this writer to another buffer. */
-	void CopyTo(void* Buffer) const;
+	UE_API void CopyTo(void* Buffer) const;
 
 	/** Gets an output buffer for writing. */
-	FMutableMemoryView GetOutputBuffer(size_t UsedSize, size_t DesiredSize);
+	UE_API FMutableMemoryView GetOutputBuffer(size_t UsedSize, size_t DesiredSize);
 
 	/** Increase the length of the written data. */
-	void Advance(size_t Size);
+	UE_API void Advance(size_t Size);
 
 private:
 	struct FChunk;
@@ -43,3 +45,5 @@ private:
 	TArray<FChunk> Chunks;
 	size_t WrittenLength;
 };
+
+#undef UE_API
diff --git a/Engine/Source/Developer/Horde/Public/Storage/Clients/BundleStorageClient.h b/Engine/Source/Developer/Horde/Public/Storage/Clients/BundleStorageClient.h
--- a/Engine/Source/Developer/Horde/Public/Storage/Clients/BundleStorageClient.h
+++ b/Engine/Source/Developer/Horde/Public/Storage/Clients/BundleStorageClient.h
@@ -8,31 +8,33 @@
 #include "Storage/Bundles/BundleOptions.h"
 #include "Storage/Clients/KeyValueStorageClient.h"
 
+#define UE_API HORDE_API
+
 /**
  * Base class for an implementation of <see cref="IStorageClient"/>, providing implementations for some common functionality using bundles.
  */
-class HORDE_API FBundleStorageClient final : public FKeyValueStorageClient
+class FBundleStorageClient final : public FKeyValueStorageClient
 {
 public:
-	FBundleStorageClient(TSharedRef<FKeyValueStorageClient> InInner);
-	virtual ~FBundleStorageClient() override;
+	UE_API FBundleStorageClient(TSharedRef<FKeyValueStorageClient> InInner);
+	UE_API virtual ~FBundleStorageClient() override;
 
 	// FKeyValueStorageClient implementation
-	virtual FBlob ReadBlob(const FBlobLocator& Locator) const override;
-	virtual FBlobHandle WriteBlob(const FUtf8StringView& BasePath, const FBlobType& Type, const TArrayView<FMemoryView>& Data, const TArrayView<FBlobHandle>& Imports) override;
+	UE_API virtual FBlob ReadBlob(const FBlobLocator& Locator) const override;
+	UE_API virtual FBlobHandle WriteBlob(const FUtf8StringView& BasePath, const FBlobType& Type, const TArrayView<FMemoryView>& Data, const TArrayView<FBlobHandle>& Imports) override;
 
 	// FBundleStorageClient implementation
-	virtual FBlobHandle CreateHandle(const FBlobLocator& Locator) const override;
-	virtual TUniquePtr<FBlobWriter> CreateWriter(FUtf8String BasePath) override;
-	TUniquePtr<FBlobWriter> CreateWriter(FUtf8String BasePath, const FBundleOptions& Options);
+	UE_API virtual FBlobHandle CreateHandle(const FBlobLocator& Locator) const override;
+	UE_API virtual TUniquePtr<FBlobWriter> CreateWriter(FUtf8String BasePath) override;
+	UE_API TUniquePtr<FBlobWriter> CreateWriter(FUtf8String BasePath, const FBundleOptions& Options);
 
-	virtual void AddAlias(const char* Name, FBlobAlias Alias) override;
-	virtual void RemoveAlias(const char* Name, FBlobHandle Handle) override;
-	virtual void FindAliases(const char* Name, int MaxResults, TArray<FBlobAlias>& OutAliases) override;
+	UE_API virtual void AddAlias(const char* Name, FBlobAlias Alias) override;
+	UE_API virtual void RemoveAlias(const char* Name, FBlobHandle Handle) override;
+	UE_API virtual void FindAliases(const char* Name, int MaxResults, TArray<FBlobAlias>& OutAliases) override;
 
-	virtual bool DeleteRef(const FRefName& Name) override;
-	virtual FBlobHandle ReadRef(const FRefName& Name, const FRefCacheTime& CacheTime) const override;
-	virtual void WriteRef(const FRefName& Name, const FBlobHandle& Handle, const FRefOptions& Options) override;
+	UE_API virtual bool DeleteRef(const FRefName& Name) override;
+	UE_API virtual FBlobHandle ReadRef(const FRefName& Name, const FRefCacheTime& CacheTime) const override;
+	UE_API virtual void WriteRef(const FRefName& Name, const FBlobHandle& Handle, const FRefOptions& Options) override;
 
 private:
 	class FBundleHandleData;
@@ -40,6 +42,8 @@ private:
 
 	TSharedRef<FKeyValueStorageClient> Inner;
 
-	TBlobHandle<FBundleHandleData> CreateBundleHandle(const FBlobLocator& Locator) const;
-	FBlobHandle GetFragmentHandle(const FBlobHandle& BundleHandle, const FUtf8StringView& Fragment) const;
+	UE_API TBlobHandle<FBundleHandleData> CreateBundleHandle(const FBlobLocator& Locator) const;
+	UE_API FBlobHandle GetFragmentHandle(const FBlobHandle& BundleHandle, const FUtf8StringView& Fragment) const;
 };
+
+#undef UE_API
diff --git a/Engine/Source/Developer/Horde/Public/Storage/Clients/FileStorageClient.h b/Engine/Source/Developer/Horde/Public/Storage/Clients/FileStorageClient.h
--- a/Engine/Source/Developer/Horde/Public/Storage/Clients/FileStorageClient.h
+++ b/Engine/Source/Developer/Horde/Public/Storage/Clients/FileStorageClient.h
@@ -6,33 +6,37 @@
 #include "KeyValueStorageClient.h"
 #include <filesystem>
 
+#define UE_API HORDE_API
+
 /**
  * Implementation of FStorageClient which writes data to files on disk.
  */
-class HORDE_API FFileStorageClient : public FKeyValueStorageClient
+class FFileStorageClient : public FKeyValueStorageClient
 {
 public:
-	FFileStorageClient(std::filesystem::path InRootDir);
-	~FFileStorageClient();
+	UE_API FFileStorageClient(std::filesystem::path InRootDir);
+	UE_API ~FFileStorageClient();
 
-	static FBlobLocator ReadRefFromFile(const std::filesystem::path &File);
-	static void WriteRefToFile(const std::filesystem::path& File, const FBlobLocator& Locator);
+	static UE_API FBlobLocator ReadRefFromFile(const std::filesystem::path &File);
+	static UE_API void WriteRefToFile(const std::filesystem::path& File, const FBlobLocator& Locator);
 
-	std::filesystem::path GetBlobFile(const FBlobLocator& Locator) const;
-	std::filesystem::path GetRefFile(const FRefName& Name) const;
+	UE_API std::filesystem::path GetBlobFile(const FBlobLocator& Locator) const;
+	UE_API std::filesystem::path GetRefFile(const FRefName& Name) const;
 
-	virtual void AddAlias(const char* Name, FBlobAlias Alias) override;
-	virtual void RemoveAlias(const char* Name, FBlobHandle Handle) override;
-	virtual void FindAliases(const char* Name, int MaxResults, TArray<FBlobAlias>& OutAliases) override;
+	UE_API virtual void AddAlias(const char* Name, FBlobAlias Alias) override;
+	UE_API virtual void RemoveAlias(const char* Name, FBlobHandle Handle) override;
+	UE_API virtual void FindAliases(const char* Name, int MaxResults, TArray<FBlobAlias>& OutAliases) override;
 
-	virtual bool DeleteRef(const FRefName& Name) override;
-	virtual FBlobHandle ReadRef(const FRefName& Name, const FRefCacheTime& CacheTime) const override;
-	virtual void WriteRef(const FRefName& Name, const FBlobHandle& Handle, const FRefOptions& Options) override;
+	UE_API virtual bool DeleteRef(const FRefName& Name) override;
+	UE_API virtual FBlobHandle ReadRef(const FRefName& Name, const FRefCacheTime& CacheTime) const override;
+	UE_API virtual void WriteRef(const FRefName& Name, const FBlobHandle& Handle, const FRefOptions& Options) override;
 
 protected:
-	virtual FBlob ReadBlob(const FBlobLocator& Locator) const override;
-	virtual FBlobHandle WriteBlob(const FUtf8StringView& BasePath, const FBlobType& Type, const TArrayView<FMemoryView>& Data, const TArrayView<FBlobHandle>& Imports) override;
+	UE_API virtual FBlob ReadBlob(const FBlobLocator& Locator) const override;
+	UE_API virtual FBlobHandle WriteBlob(const FUtf8StringView& BasePath, const FBlobType& Type, const TArrayView<FMemoryView>& Data, const TArrayView<FBlobHandle>& Imports) override;
 
 private:
 	std::filesystem::path RootDir;
 };
+
+#undef UE_API
diff --git a/Engine/Source/Developer/Horde/Public/Storage/Clients/KeyValueStorageClient.h b/Engine/Source/Developer/Horde/Public/Storage/Clients/KeyValueStorageClient.h
--- a/Engine/Source/Developer/Horde/Public/Storage/Clients/KeyValueStorageClient.h
+++ b/Engine/Source/Developer/Horde/Public/Storage/Clients/KeyValueStorageClient.h
@@ -6,10 +6,12 @@
 #include "Storage/StorageClient.h"
 #include "Storage/BlobHandle.h"
 
+#define UE_API HORDE_API
+
 /*
  * Base class for storage clients that wrap a diirect key/value type store without any merging/splitting.
  */
-class HORDE_API FKeyValueStorageClient : public FStorageClient
+class FKeyValueStorageClient : public FStorageClient
 {
 public:
 	/** Read a single blob from the underlying store. */
@@ -19,9 +21,11 @@ public:
 	virtual FBlobHandle WriteBlob(const FUtf8StringView& BasePath, const FBlobType& Type, const TArrayView<FMemoryView>& Data, const TArrayView<FBlobHandle>& Imports) = 0;
 
 	// Implementation of FStorageClient
-	virtual FBlobHandle CreateHandle(const FBlobLocator& Locator) const override;
-	virtual TUniquePtr<FBlobWriter> CreateWriter(FUtf8String BasePath = FUtf8String()) override;
+	UE_API virtual FBlobHandle CreateHandle(const FBlobLocator& Locator) const override;
+	UE_API virtual TUniquePtr<FBlobWriter> CreateWriter(FUtf8String BasePath = FUtf8String()) override;
 
 private:
 	class FHandleData;
 };
+
+#undef UE_API
diff --git a/Engine/Source/Developer/Horde/Public/Storage/Clients/MemoryStorageClient.h b/Engine/Source/Developer/Horde/Public/Storage/Clients/MemoryStorageClient.h
--- a/Engine/Source/Developer/Horde/Public/Storage/Clients/MemoryStorageClient.h
+++ b/Engine/Source/Developer/Horde/Public/Storage/Clients/MemoryStorageClient.h
@@ -5,26 +5,28 @@
 #include "CoreMinimal.h"
 #include "KeyValueStorageClient.h"
 
+#define UE_API HORDE_API
+
 /**
  * Implementation of FStorageClient which stores data in memory.
  */
-class HORDE_API FMemoryStorageClient final : public FKeyValueStorageClient
+class FMemoryStorageClient final : public FKeyValueStorageClient
 {
 public:
-	FMemoryStorageClient();
-	~FMemoryStorageClient();
+	UE_API FMemoryStorageClient();
+	UE_API ~FMemoryStorageClient();
 
-	virtual void AddAlias(const char* Name, FBlobAlias Alias) override;
-	virtual void RemoveAlias(const char* Name, FBlobHandle Handle) override;
-	virtual void FindAliases(const char* Name, int MaxResults, TArray<FBlobAlias>& OutAliases) override;
+	UE_API virtual void AddAlias(const char* Name, FBlobAlias Alias) override;
+	UE_API virtual void RemoveAlias(const char* Name, FBlobHandle Handle) override;
+	UE_API virtual void FindAliases(const char* Name, int MaxResults, TArray<FBlobAlias>& OutAliases) override;
 
-	virtual bool DeleteRef(const FRefName& Name) override;
-	virtual FBlobHandle ReadRef(const FRefName& Name, const FRefCacheTime& CacheTime) const override;
-	virtual void WriteRef(const FRefName& Name, const FBlobHandle& Target, const FRefOptions& Options) override;
+	UE_API virtual bool DeleteRef(const FRefName& Name) override;
+	UE_API virtual FBlobHandle ReadRef(const FRefName& Name, const FRefCacheTime& CacheTime) const override;
+	UE_API virtual void WriteRef(const FRefName& Name, const FBlobHandle& Target, const FRefOptions& Options) override;
 
 protected:
-	virtual FBlob ReadBlob(const FBlobLocator& Locator) const override;
-	virtual FBlobHandle WriteBlob(const FUtf8StringView& BasePath, const FBlobType& Type, const TArrayView<FMemoryView>& Data, const TArrayView<FBlobHandle>& Imports) override;
+	UE_API virtual FBlob ReadBlob(const FBlobLocator& Locator) const override;
+	UE_API virtual FBlobHandle WriteBlob(const FUtf8StringView& BasePath, const FBlobType& Type, const TArrayView<FMemoryView>& Data, const TArrayView<FBlobHandle>& Imports) override;
 
 private:
 	mutable FCriticalSection CriticalSection;
@@ -33,3 +35,5 @@ private:
 	TMap<FRefName, FBlobLocator> Refs;
 	TMap<FUtf8String, TArray<FBlobAlias>> Aliases;
 };
+
+#undef UE_API
diff --git a/Engine/Source/Developer/Horde/Public/Storage/Nodes/ChunkNode.h b/Engine/Source/Developer/Horde/Public/Storage/Nodes/ChunkNode.h
--- a/Engine/Source/Developer/Horde/Public/Storage/Nodes/ChunkNode.h
+++ b/Engine/Source/Developer/Horde/Public/Storage/Nodes/ChunkNode.h
@@ -11,13 +11,15 @@
 #include "Hash/Blake3.h"
 #include "Hash/BuzHash.h"
 
+#define UE_API HORDE_API
+
 /**
  * Options for chunking data
  */
 
-struct HORDE_API FChunkingOptions
+struct FChunkingOptions
 {
-	static const FChunkingOptions Default;
+	static UE_API const FChunkingOptions Default;
 
 	int32 MinChunkSize = 4 * 1024;
 	int32 TargetChunkSize = 64 * 1024;
@@ -27,40 +29,40 @@ struct HORDE_API FChunkingOptions
 /**
  * Node containing a chunk of data
  */
-class HORDE_API FChunkNode
+class FChunkNode
 {
 public:
-	static const FBlobType LeafBlobType;
-	static const FBlobType InteriorBlobType;
+	static UE_API const FBlobType LeafBlobType;
+	static UE_API const FBlobType InteriorBlobType;
 
 	TArray<FBlobHandleWithHash> Children;
 	FSharedBufferView Data;
 
-	FChunkNode();
-	FChunkNode(TArray<FBlobHandleWithHash> InChildren, FSharedBufferView InData);
-	~FChunkNode();
+	UE_API FChunkNode();
+	UE_API FChunkNode(TArray<FBlobHandleWithHash> InChildren, FSharedBufferView InData);
+	UE_API ~FChunkNode();
 
-	static FChunkNode Read(FBlob Blob);
-	FBlobHandleWithHash Write(FBlobWriter& Writer) const;
+	static UE_API FChunkNode Read(FBlob Blob);
+	UE_API FBlobHandleWithHash Write(FBlobWriter& Writer) const;
 
-	static FBlobHandleWithHash Write(FBlobWriter& Writer, const TArrayView<const FBlobHandleWithHash>& Children, FMemoryView Data);
+	static UE_API FBlobHandleWithHash Write(FBlobWriter& Writer, const TArrayView<const FBlobHandleWithHash>& Children, FMemoryView Data);
 };
 
 /**
  * Utility class for reading data a data stream from a tree of chunk nodes
  */
-class HORDE_API FChunkNodeReader
+class FChunkNodeReader
 {
 public:
-	FChunkNodeReader(FBlob Blob);
-	FChunkNodeReader(const FBlobHandle& Handle);
-	~FChunkNodeReader();
+	UE_API FChunkNodeReader(FBlob Blob);
+	UE_API FChunkNodeReader(const FBlobHandle& Handle);
+	UE_API ~FChunkNodeReader();
 
-	bool IsComplete() const;
-	FMemoryView GetBuffer() const;
-	void Advance(int32 Length);
+	UE_API bool IsComplete() const;
+	UE_API FMemoryView GetBuffer() const;
+	UE_API void Advance(int32 Length);
 
-	operator bool() const;
+	UE_API operator bool() const;
 
 private:
 	struct FStackEntry
@@ -79,15 +81,15 @@ private:
 /**
  * Utility class for writing new data to a tree of chunk nodes
  */
-class HORDE_API FChunkNodeWriter
+class FChunkNodeWriter
 {
 public:
-	FChunkNodeWriter(FBlobWriter& InWriter, const FChunkingOptions& InOptions = FChunkingOptions::Default);
-	~FChunkNodeWriter();
+	UE_API FChunkNodeWriter(FBlobWriter& InWriter, const FChunkingOptions& InOptions = FChunkingOptions::Default);
+	UE_API ~FChunkNodeWriter();
 
-	void Write(FMemoryView Data);
+	UE_API void Write(FMemoryView Data);
 
-	FBlobHandleWithHash Flush(FIoHash& OutStreamHash);
+	UE_API FBlobHandleWithHash Flush(FIoHash& OutStreamHash);
 
 private:
 	FBlobWriter& Writer;
@@ -101,5 +103,7 @@ private:
 
 	TArray<FBlobHandleWithHash> Nodes;
 
-	void WriteNode();
+	UE_API void WriteNode();
 };
+
+#undef UE_API
diff --git a/Engine/Source/Developer/Horde/Public/Storage/Nodes/DirectoryEntry.h b/Engine/Source/Developer/Horde/Public/Storage/Nodes/DirectoryEntry.h
--- a/Engine/Source/Developer/Horde/Public/Storage/Nodes/DirectoryEntry.h
+++ b/Engine/Source/Developer/Horde/Public/Storage/Nodes/DirectoryEntry.h
@@ -7,13 +7,15 @@
 #include "../BlobHandle.h"
 #include "IO/IoHash.h"
 
+#define UE_API HORDE_API
+
 class FBlobReader;
 class FBlobWriter;
 
 /**
  * Entry for a directory within a directory node
  */
-class HORDE_API FDirectoryEntry
+class FDirectoryEntry
 {
 public:
 	FBlobHandle Target;
@@ -25,9 +27,11 @@ public:
 	/** Name of this directory. */
 	const FUtf8String Name;
 
-	FDirectoryEntry(FBlobHandle InTarget, const FIoHash& InTargetHash, FUtf8String InName, int64 InLength);
-	~FDirectoryEntry();
+	UE_API FDirectoryEntry(FBlobHandle InTarget, const FIoHash& InTargetHash, FUtf8String InName, int64 InLength);
+	UE_API ~FDirectoryEntry();
 
-	static FDirectoryEntry Read(FBlobReader& Reader);
-	void Write(FBlobWriter& Writer) const;
+	static UE_API FDirectoryEntry Read(FBlobReader& Reader);
+	UE_API void Write(FBlobWriter& Writer) const;
 };
+
+#undef UE_API
diff --git a/Engine/Source/Developer/Horde/Public/Storage/Nodes/DirectoryNode.h b/Engine/Source/Developer/Horde/Public/Storage/Nodes/DirectoryNode.h
--- a/Engine/Source/Developer/Horde/Public/Storage/Nodes/DirectoryNode.h
+++ b/Engine/Source/Developer/Horde/Public/Storage/Nodes/DirectoryNode.h
@@ -8,6 +8,8 @@
 #include "../BlobType.h"
 #include "../StorageClient.h"
 
+#define UE_API HORDE_API
+
 /**
  * Flags for a directory node
  */
@@ -20,18 +22,20 @@ enum class EDirectoryFlags
 /**
  * Stores the contents of a directory in a blob
  */
-class HORDE_API FDirectoryNode
+class FDirectoryNode
 {
 public:
-	static const FBlobType BlobType;
+	static UE_API const FBlobType BlobType;
 
 	EDirectoryFlags Flags;
 	TMap<FUtf8String, FFileEntry> NameToFile;
 	TMap<FUtf8String, FDirectoryEntry> NameToDirectory;
 
-	FDirectoryNode(EDirectoryFlags InFlags = EDirectoryFlags::None);
-	~FDirectoryNode();
+	UE_API FDirectoryNode(EDirectoryFlags InFlags = EDirectoryFlags::None);
+	UE_API ~FDirectoryNode();
 
-	static FDirectoryNode Read(const FBlob& Blob);
-	FBlobHandle Write(FBlobWriter& Writer) const;
+	static UE_API FDirectoryNode Read(const FBlob& Blob);
+	UE_API FBlobHandle Write(FBlobWriter& Writer) const;
 };
+
+#undef UE_API
diff --git a/Engine/Source/Developer/Horde/Public/Storage/Nodes/FileEntry.h b/Engine/Source/Developer/Horde/Public/Storage/Nodes/FileEntry.h
--- a/Engine/Source/Developer/Horde/Public/Storage/Nodes/FileEntry.h
+++ b/Engine/Source/Developer/Horde/Public/Storage/Nodes/FileEntry.h
@@ -8,6 +8,8 @@
 #include "SharedBufferView.h"
 #include "Storage/BlobHandle.h"
 
+#define UE_API HORDE_API
+
 
 class FBlobReader;
 class FBlobWriter;
@@ -41,7 +43,7 @@ ENUM_CLASS_FLAGS(EFileEntryFlags)
 /**
  * Entry for a file within a directory node
  */
-class HORDE_API FFileEntry
+class FFileEntry
 {
 public:
 	/** Handle to the root chunked data node. */
@@ -62,9 +64,11 @@ public:
 	/** Custom user data for this file entry. */
 	FSharedBufferView CustomData;
 
-	FFileEntry(FBlobHandleWithHash InTarget, FUtf8String InName, EFileEntryFlags InFlags, int64 InLength, const FIoHash& InHash, FSharedBufferView InCustomData);
-	~FFileEntry();
+	UE_API FFileEntry(FBlobHandleWithHash InTarget, FUtf8String InName, EFileEntryFlags InFlags, int64 InLength, const FIoHash& InHash, FSharedBufferView InCustomData);
+	UE_API ~FFileEntry();
 
-	static FFileEntry Read(FBlobReader& Reader);
-	void Write(FBlobWriter& Writer) const;
+	static UE_API FFileEntry Read(FBlobReader& Reader);
+	UE_API void Write(FBlobWriter& Writer) const;
 };
+
+#undef UE_API
diff --git a/Engine/Source/Developer/Horde/Public/Storage/RefName.h b/Engine/Source/Developer/Horde/Public/Storage/RefName.h
--- a/Engine/Source/Developer/Horde/Public/Storage/RefName.h
+++ b/Engine/Source/Developer/Horde/Public/Storage/RefName.h
@@ -5,20 +5,24 @@
 #include "CoreMinimal.h"
 #include "Containers/Utf8String.h"
 
+#define UE_API HORDE_API
+
 // Identifier for a ref in the storage system. Refs serve as GC roots, and are persistent entry points to expanding data structures within the store.
-struct HORDE_API FRefName
+struct FRefName
 {
 public:
-	FRefName(FUtf8String Text);
-	~FRefName();
+	UE_API FRefName(FUtf8String Text);
+	UE_API ~FRefName();
 
 	/** Accessor for the underlying string. */
-	const FUtf8String& GetText() const;
+	UE_API const FUtf8String& GetText() const;
 
-	bool operator==(const FRefName& Other) const;
-	bool operator!=(const FRefName& Other) const;
+	UE_API bool operator==(const FRefName& Other) const;
+	UE_API bool operator!=(const FRefName& Other) const;
 	friend uint32 GetTypeHash(const FRefName& RefName);
 
 private:
 	FUtf8String Text;
 };
+
+#undef UE_API
diff --git a/Engine/Source/Developer/Horde/Public/Storage/RefOptions.h b/Engine/Source/Developer/Horde/Public/Storage/RefOptions.h
--- a/Engine/Source/Developer/Horde/Public/Storage/RefOptions.h
+++ b/Engine/Source/Developer/Horde/Public/Storage/RefOptions.h
@@ -4,12 +4,14 @@
 
 #include "CoreMinimal.h"
 
+#define UE_API HORDE_API
+
 /**
  * Options for adding a new ref
  */ 
-struct HORDE_API FRefOptions
+struct FRefOptions
 {
-	static const FRefOptions Default;
+	static UE_API const FRefOptions Default;
 
 	/** Time until a ref is expired. If zero, the ref does not expire. */
 	FTimespan Lifetime;
@@ -17,5 +19,7 @@ struct HORDE_API FRefOptions
 	/** Whether to extend the remaining lifetime of a ref whenever it is fetched. Defaults to true. */
 	bool bExtend = true;
 
-	FRefOptions();
+	UE_API FRefOptions();
 };
+
+#undef UE_API
diff --git a/Engine/Source/Developer/Horde/Public/Storage/StorageClient.h b/Engine/Source/Developer/Horde/Public/Storage/StorageClient.h
--- a/Engine/Source/Developer/Horde/Public/Storage/StorageClient.h
+++ b/Engine/Source/Developer/Horde/Public/Storage/StorageClient.h
@@ -9,6 +9,8 @@
 #include "RefName.h"
 #include "Templates/SharedPointer.h"
 
+#define UE_API HORDE_API
+
 class FBlobAlias;
 struct FAliasInfo;
 class FBlobWriter;
@@ -29,10 +31,10 @@ struct FRefCacheTime
 /**
  * Interface for the storage system.
  */
-class HORDE_API FStorageClient : public TSharedFromThis<FStorageClient, ESPMode::ThreadSafe>
+class FStorageClient : public TSharedFromThis<FStorageClient, ESPMode::ThreadSafe>
 {
 public:
-	virtual ~FStorageClient();
+	UE_API virtual ~FStorageClient();
 
 	//
 	// Blobs
@@ -68,13 +70,13 @@ public:
 	virtual FBlobHandle ReadRef(const FRefName& Name, const FRefCacheTime& CacheTime) const = 0;
 
 	/** Writes a new ref to the store. */
-	void WriteRef(const FRefName& Name, const FBlobHandle& Handle);
+	UE_API void WriteRef(const FRefName& Name, const FBlobHandle& Handle);
 	virtual void WriteRef(const FRefName& Name, const FBlobHandle& Handle, const FRefOptions& Options) = 0;
 
 protected:
 
 	/** Creates a new locator with a particular base path. */
-	static FBlobLocator CreateLocator(const FUtf8StringView& BasePath);
+	static UE_API FBlobLocator CreateLocator(const FUtf8StringView& BasePath);
 };
 
 /**
@@ -115,3 +117,4 @@ struct FAliasInfo
 	~FAliasInfo();
 };
 
+#undef UE_API
diff --git a/Engine/Source/Developer/LogVisualizer/Private/VisualLoggerRenderingActor.h b/Engine/Source/Developer/LogVisualizer/Private/VisualLoggerRenderingActor.h
--- a/Engine/Source/Developer/LogVisualizer/Private/VisualLoggerRenderingActor.h
+++ b/Engine/Source/Developer/LogVisualizer/Private/VisualLoggerRenderingActor.h
@@ -5,6 +5,8 @@
 #include "VisualLoggerRenderingActorBase.h"
 #include "VisualLoggerRenderingActor.generated.h"
 
+#define UE_API LOGVISUALIZER_API
+
 // enable adding some hard coded shapes to the VisualLoggerRenderingActor for testing
 #define VLOG_TEST_DEBUG_RENDERING 0
 
@@ -17,27 +19,29 @@ struct FVisualLoggerDBRow;
 
 DECLARE_MULTICAST_DELEGATE_OneParam(FOnSelectionChanged, class AActor*);
 
-UCLASS(config = Engine, NotBlueprintable, Transient, notplaceable, AdvancedClassDisplay)
-class LOGVISUALIZER_API AVisualLoggerRenderingActor : public AVisualLoggerRenderingActorBase 
+UCLASS(MinimalAPI, config = Engine, NotBlueprintable, Transient, notplaceable, AdvancedClassDisplay)
+class AVisualLoggerRenderingActor : public AVisualLoggerRenderingActorBase 
 {
 public:
 	GENERATED_UCLASS_BODY()
-	virtual ~AVisualLoggerRenderingActor();
-	void ResetRendering();
-	void ObjectVisibilityChanged(const FName& RowName);
-	void ObjectSelectionChanged(const TArray<FName>& Selection);
-	void OnItemSelectionChanged(const FVisualLoggerDBRow& BDRow, int32 ItemIndex);
-
-	virtual void IterateDebugShapes(TFunction<void(const FTimelineDebugShapes&) > Callback) override;
-	virtual bool MatchCategoryFilters(const FName& CategoryName, ELogVerbosity::Type Verbosity) const override;
+	UE_API virtual ~AVisualLoggerRenderingActor();
+	UE_API void ResetRendering();
+	UE_API void ObjectVisibilityChanged(const FName& RowName);
+	UE_API void ObjectSelectionChanged(const TArray<FName>& Selection);
+	UE_API void OnItemSelectionChanged(const FVisualLoggerDBRow& BDRow, int32 ItemIndex);
+
+	UE_API virtual void IterateDebugShapes(TFunction<void(const FTimelineDebugShapes&) > Callback) override;
+	UE_API virtual bool MatchCategoryFilters(const FName& CategoryName, ELogVerbosity::Type Verbosity) const override;
 private:
-	void OnFiltersChanged();
+	UE_API void OnFiltersChanged();
 
 	TArray<FName> CachedRowSelection;
 	TMap<FName, FTimelineDebugShapes> DebugShapesPerRow;
 
 #if VLOG_TEST_DEBUG_RENDERING
-	void AddDebugRendering();
+	UE_API void AddDebugRendering();
 	FTimelineDebugShapes TestDebugShapes;
 #endif
 };
+
+#undef UE_API
diff --git a/Engine/Source/Developer/LogVisualizer/Public/LogVisualizerPublic.h b/Engine/Source/Developer/LogVisualizer/Public/LogVisualizerPublic.h
--- a/Engine/Source/Developer/LogVisualizer/Public/LogVisualizerPublic.h
+++ b/Engine/Source/Developer/LogVisualizer/Public/LogVisualizerPublic.h
@@ -5,6 +5,8 @@
 *****************************************************************************/
 #include "CoreMinimal.h"
 
+#define UE_API LOGVISUALIZER_API
+
 class FText;
 struct FGeometry;
 struct FKeyEvent;
@@ -31,14 +33,14 @@ struct FVisualLoggerEvents
 };
 
 class FVisualLoggerTimeSliderController;
-struct LOGVISUALIZER_API FLogVisualizer
+struct FLogVisualizer
 {
 	/** LogVisualizer interface*/
-	void Reset();
+	UE_API void Reset();
 
-	FLinearColor GetColorForCategory(int32 Index) const;
-	FLinearColor GetColorForCategory(const FString& InFilterName) const;
-	UWorld* GetWorld(UObject* OptionalObject = nullptr);
+	UE_API FLinearColor GetColorForCategory(int32 Index) const;
+	UE_API FLinearColor GetColorForCategory(const FString& InFilterName) const;
+	UE_API UWorld* GetWorld(UObject* OptionalObject = nullptr);
 	FVisualLoggerEvents& GetEvents() { return VisualLoggerEvents; }
 
 	void SetCurrentVisualizer(TSharedPtr<class SVisualLogger> Visualizer) { CurrentVisualizer = Visualizer; }
@@ -53,30 +55,30 @@ struct LOGVISUALIZER_API FLogVisualizer
 		return AnimationOutlinerFillPercentage;
 	}
 
-	int32 GetNextItem(FName RowName, int32 MoveDistance = 1);
-	int32 GetPreviousItem(FName RowName, int32 MoveDistance = 1);
+	UE_API int32 GetNextItem(FName RowName, int32 MoveDistance = 1);
+	UE_API int32 GetPreviousItem(FName RowName, int32 MoveDistance = 1);
 
 	// @todo: This function currently doesn't do anything!
-	void GotoNextItem(FName RowName, int32 MoveDistance = 1);
+	UE_API void GotoNextItem(FName RowName, int32 MoveDistance = 1);
 	// @todo: This function currently doesn't do anything!
-	void GotoPreviousItem(FName RowName, int32 MoveDistance = 1);
-	void GotoFirstItem(FName RowName);
-	void GotoLastItem(FName RowName);
-	void GotoLastItemAnyRow();
+	UE_API void GotoPreviousItem(FName RowName, int32 MoveDistance = 1);
+	UE_API void GotoFirstItem(FName RowName);
+	UE_API void GotoLastItem(FName RowName);
+	UE_API void GotoLastItemAnyRow();
 
-	void UpdateCameraPosition(FName Rowname, int32 ItemIndes);
+	UE_API void UpdateCameraPosition(FName Rowname, int32 ItemIndes);
 
-	void SeekToTime(float Time);
+	UE_API void SeekToTime(float Time);
 
 	/** Static access */
-	static void Initialize();
-	static void Shutdown();
-	static FLogVisualizer& Get();
+	static UE_API void Initialize();
+	static UE_API void Shutdown();
+	static UE_API FLogVisualizer& Get();
 
 protected:
 	TSharedPtr<FVisualLoggerTimeSliderController> GetTimeSliderController() { return TimeSliderController; }
 
-	static TSharedPtr< struct FLogVisualizer > StaticInstance;
+	static UE_API TSharedPtr< struct FLogVisualizer > StaticInstance;
 
 	TSharedPtr<FVisualLoggerTimeSliderController> TimeSliderController;
 	FVisualLoggerEvents VisualLoggerEvents;
@@ -87,4 +89,6 @@ protected:
 	friend class SVisualLoggerView;
 	friend class SVisualLogger;
 	friend struct FVisualLoggerCanvasRenderer;
-};
\ No newline at end of file
+};
+
+#undef UE_API
diff --git a/Engine/Source/Developer/LogVisualizer/Public/LogVisualizerSessionSettings.h b/Engine/Source/Developer/LogVisualizer/Public/LogVisualizerSessionSettings.h
--- a/Engine/Source/Developer/LogVisualizer/Public/LogVisualizerSessionSettings.h
+++ b/Engine/Source/Developer/LogVisualizer/Public/LogVisualizerSessionSettings.h
@@ -11,8 +11,10 @@ GameplayDebuggerSettings.h: Declares the UGameplayDebuggerSettings class.
 #include "UObject/Object.h"
 #include "LogVisualizerSessionSettings.generated.h"
 
-UCLASS()
-class LOGVISUALIZER_API ULogVisualizerSessionSettings : public UObject
+#define UE_API LOGVISUALIZER_API
+
+UCLASS(MinimalAPI)
+class ULogVisualizerSessionSettings : public UObject
 {
 	GENERATED_UCLASS_BODY()
 public:
@@ -25,7 +27,7 @@ public:
 
 	// UObject overrides
 #if WITH_EDITOR
-	virtual void PostEditChangeProperty(struct FPropertyChangedEvent& PropertyChangedEvent) override;
+	UE_API virtual void PostEditChangeProperty(struct FPropertyChangedEvent& PropertyChangedEvent) override;
 #endif
 
 private:
@@ -33,3 +35,5 @@ private:
 	// Holds an event delegate that is executed when a setting has changed.
 	FSettingChangedEvent SettingChangedEvent;
 };
+
+#undef UE_API
diff --git a/Engine/Source/Developer/LogVisualizer/Public/LogVisualizerSettings.h b/Engine/Source/Developer/LogVisualizer/Public/LogVisualizerSettings.h
--- a/Engine/Source/Developer/LogVisualizer/Public/LogVisualizerSettings.h
+++ b/Engine/Source/Developer/LogVisualizer/Public/LogVisualizerSettings.h
@@ -148,8 +148,8 @@ private:
 
 struct FCategoryFiltersManager;
 
-UCLASS(config = EditorPerProjectUserSettings)
-class LOGVISUALIZER_API ULogVisualizerSettings : public UObject
+UCLASS(MinimalAPI, config = EditorPerProjectUserSettings)
+class ULogVisualizerSettings : public UObject
 {
 	GENERATED_UCLASS_BODY()
 	friend struct FCategoryFiltersManager;
@@ -224,18 +224,18 @@ public:
 
 	// UObject overrides
 #if WITH_EDITOR
-	virtual void PostEditChangeProperty(struct FPropertyChangedEvent& PropertyChangedEvent) override;
+	LOGVISUALIZER_API virtual void PostEditChangeProperty(struct FPropertyChangedEvent& PropertyChangedEvent) override;
 #endif
 
-	class UMaterial* GetDebugMeshMaterial();
+	LOGVISUALIZER_API class UMaterial* GetDebugMeshMaterial();
 
-	void SavePersistentData();
+	LOGVISUALIZER_API void SavePersistentData();
 
-	void ClearPersistentData();
+	LOGVISUALIZER_API void ClearPersistentData();
 
-	void LoadPersistentData();
+	LOGVISUALIZER_API void LoadPersistentData();
 
-	void ConfigureVisLog();
+	LOGVISUALIZER_API void ConfigureVisLog();
 
 protected:
 	UPROPERTY(config)
diff --git a/Engine/Source/Developer/LogVisualizer/Public/VisualLoggerRenderingActorBase.h b/Engine/Source/Developer/LogVisualizer/Public/VisualLoggerRenderingActorBase.h
--- a/Engine/Source/Developer/LogVisualizer/Public/VisualLoggerRenderingActorBase.h
+++ b/Engine/Source/Developer/LogVisualizer/Public/VisualLoggerRenderingActorBase.h
@@ -9,6 +9,8 @@
 #include "DebugRenderSceneProxy.h"
 #include "VisualLoggerRenderingActorBase.generated.h"
 
+#define UE_API LOGVISUALIZER_API
+
 class UVisualLoggerRenderingComponent;
 
 /**
@@ -16,12 +18,12 @@ class UVisualLoggerRenderingComponent;
 *	Base class shared between RewindDebugger and Visual Logger implementations
 */
 
-UCLASS(config = Engine, NotBlueprintable, Transient, notplaceable, AdvancedClassDisplay)
-class LOGVISUALIZER_API AVisualLoggerRenderingActorBase : public AActor
+UCLASS(MinimalAPI, config = Engine, NotBlueprintable, Transient, notplaceable, AdvancedClassDisplay)
+class AVisualLoggerRenderingActorBase : public AActor
 {
 public:
 	GENERATED_UCLASS_BODY()
-	virtual ~AVisualLoggerRenderingActorBase();
+	UE_API virtual ~AVisualLoggerRenderingActorBase();
 
 	struct FTimelineDebugShapes
 	{
@@ -64,13 +66,15 @@ public:
 #endif // WITH_EDITOR
 
 // Allows to override the far clipping distance for all the log elements by this actor :
-	void SetFarClippingDistance(double Distance);
-	double GetFarClippingDistance() const;
+	UE_API void SetFarClippingDistance(double Distance);
+	UE_API double GetFarClippingDistance() const;
 
 protected:
 #if ENABLE_VISUAL_LOG
-	void GetDebugShapes(const FVisualLogEntry& EntryItem, bool bAddEntryLocationPointer, FTimelineDebugShapes& OutDebugShapes);
+	UE_API void GetDebugShapes(const FVisualLogEntry& EntryItem, bool bAddEntryLocationPointer, FTimelineDebugShapes& OutDebugShapes);
 #endif
 
 	UVisualLoggerRenderingComponent* RenderingComponent = nullptr;
 };
+
+#undef UE_API
