PR #10649: In FDirectoryWatcherWindows, key the request map using full directory paths.

diff --git a/Engine/Source/Developer/DirectoryWatcher/Private/DirectoryWatcherTests.cpp b/Engine/Source/Developer/DirectoryWatcher/Private/DirectoryWatcherTests.cpp
--- a/Engine/Source/Developer/DirectoryWatcher/Private/DirectoryWatcherTests.cpp
+++ b/Engine/Source/Developer/DirectoryWatcher/Private/DirectoryWatcherTests.cpp
@@ -31,20 +31,20 @@ struct FDirectoryWatcherTestPayload
 	}
 	~FDirectoryWatcherTestPayload()
 	{
-		IFileManager::Get().DeleteDirectory(*WorkingDir, false, true);
-
 		FDirectoryWatcherModule& Module = FModuleManager::LoadModuleChecked<FDirectoryWatcherModule>(TEXT("DirectoryWatcher"));
 		if (IDirectoryWatcher* DirectoryWatcher = Module.Get())
 		{
 			DirectoryWatcher->UnregisterDirectoryChangedCallback_Handle(WorkingDir, WatcherDelegate);
 		}
+
+		IFileManager::Get().DeleteDirectory(*WorkingDir, false, true);
 	}
 
 	void OnDirectoryChanged(const TArray<FFileChangeData>& InFileChanges)
 	{
 		for (const auto& Change : InFileChanges)
 		{
-			FString RelativeFilename = *FPaths::ConvertRelativePathToFull(Change.Filename) + WorkingDir.Len();
+			FString RelativeFilename = *FPaths::ConvertRelativePathToFull(Change.Filename) + FPaths::ConvertRelativePathToFull(WorkingDir).Len();
 
 			UE_LOG(LogDirectoryWatcherTests, Log, TEXT("File '%s'. Code: %d."), *Change.Filename, (uint8)Change.Action);
 
@@ -76,9 +76,14 @@ struct FDirectoryWatcherTestPayload
 
 namespace DirectoryWatcherTests
 {
+	FString GetWorkingDirAsRelativePath()
+	{
+		return FPaths::AutomationTransientDir() / TEXT("DirectoryWatcher") / TEXT("");
+	}
+
 	FString GetWorkingDir()
 	{
-		return FPaths::ConvertRelativePathToFull(FPaths::AutomationTransientDir() / TEXT("DirectoryWatcher")) / TEXT("");
+		return FPaths::ConvertRelativePathToFull(GetWorkingDirAsRelativePath()) / TEXT("");
 	}
 }
 
@@ -340,3 +345,47 @@ bool FDirectoryWatcherIgnoreSubtreeTest::RunTest(const FString& Parameters)
 	return true;
 }
 
+// Test dual directory watchers for the same directory using both relative and absolute paths.
+IMPLEMENT_SIMPLE_AUTOMATION_TEST(FDirectoryWatcherCreateDualWatchersTest, "System.Plugins.Directory Watcher.Create Dual Watchers", EAutomationTestFlags::EditorContext | EAutomationTestFlags::EngineFilter)
+bool FDirectoryWatcherCreateDualWatchersTest::RunTest(const FString& Parameters)
+{
+	static const TCHAR* Filename = TEXT("dual.tmp");
+
+	const FString WorkingDirRelativePath = DirectoryWatcherTests::GetWorkingDirAsRelativePath();
+	const FString WorkingDirFullPath = DirectoryWatcherTests::GetWorkingDir();
+	const FString FileFullPath = WorkingDirFullPath / Filename;
+
+	// Start watching the directory using relative path
+	TSharedPtr<FDirectoryWatcherTestPayload> WatchRelativePath = MakeShareable(new FDirectoryWatcherTestPayload(WorkingDirRelativePath));
+
+	// Start watching the directory using full path
+	TSharedPtr<FDirectoryWatcherTestPayload> WatchFullPath = MakeShareable(new FDirectoryWatcherTestPayload(WorkingDirFullPath));
+
+	// Give the stream time to start up before doing the test
+	AddCommand(new FDelayedFunctionLatentCommand([this, FileFullPath, WatchRelativePath, WatchFullPath] {
+
+		// Create a file and check that it got reported as created by both watchers
+		FFileHelper::SaveStringToFile(TEXT(""), *FileFullPath);
+
+		AddCommand(new FDelayedFunctionLatentCommand([this, WatchRelativePath, WatchFullPath] {
+
+			FFileChangeData::EFileChangeAction* Action = WatchRelativePath->ReportedChanges.Find(Filename);
+
+			if (!Action || *Action != FFileChangeData::FCA_Added)
+			{
+				UE_LOG(LogDirectoryWatcherTests, Error, TEXT("New file '%s' was not correctly reported as being added by relative path watcher."), Filename);
+			}
+
+			Action = WatchFullPath->ReportedChanges.Find(Filename);
+
+			if (!Action || *Action != FFileChangeData::FCA_Added)
+			{
+				UE_LOG(LogDirectoryWatcherTests, Error, TEXT("New file '%s' was not correctly reported as being added by full path watcher."), Filename);
+			}
+
+			}, TestTickDelay));
+
+		}, TestTickDelay));
+
+	return true;
+}
diff --git a/Engine/Source/Developer/DirectoryWatcher/Private/Windows/DirectoryWatcherWindows.cpp b/Engine/Source/Developer/DirectoryWatcher/Private/Windows/DirectoryWatcherWindows.cpp
--- a/Engine/Source/Developer/DirectoryWatcher/Private/Windows/DirectoryWatcherWindows.cpp
+++ b/Engine/Source/Developer/DirectoryWatcher/Private/Windows/DirectoryWatcherWindows.cpp
@@ -42,7 +42,8 @@ FDirectoryWatcherWindows::~FDirectoryWatcherWindows()
 
 bool FDirectoryWatcherWindows::RegisterDirectoryChangedCallback_Handle( const FString& Directory, const FDirectoryChanged& InDelegate, FDelegateHandle& Handle, uint32 Flags )
 {
-	const FDirectoryWithFlags DirectoryKey(Directory, Flags);
+	const FString DirectoryFullPath = FPaths::ConvertRelativePathToFull(Directory);
+	const FDirectoryWithFlags DirectoryKey(DirectoryFullPath, Flags);
 	FDirectoryWatchRequestWindows** RequestPtr = RequestMap.Find(DirectoryKey);
 	FDirectoryWatchRequestWindows* Request = NULL;
 	
@@ -59,12 +60,12 @@ bool FDirectoryWatcherWindows::RegisterDirectoryChangedCallback_Handle( const FS
 		NumRequests++;
 
 		// Begin reading directory changes
-		if ( !Request->Init(Directory) )
+		if ( !Request->Init(DirectoryFullPath) )
 		{
 			uint32 Error = GetLastError();
 			TCHAR ErrorMsg[1024];
 			FPlatformMisc::GetSystemErrorMessage(ErrorMsg, 1024, Error);
-			UE_LOG(LogDirectoryWatcher, Warning, TEXT("Failed to begin reading directory changes for %s. Error: %s (0x%08x)"), *Directory, ErrorMsg, Error);
+			UE_LOG(LogDirectoryWatcher, Warning, TEXT("Failed to begin reading directory changes for %s. Error: %s (0x%08x)"), *DirectoryFullPath, ErrorMsg, Error);
 
 			delete Request;
 			NumRequests--;
@@ -81,9 +82,11 @@ bool FDirectoryWatcherWindows::RegisterDirectoryChangedCallback_Handle( const FS
 
 bool FDirectoryWatcherWindows::UnregisterDirectoryChangedCallback_Handle( const FString& Directory, FDelegateHandle InHandle )
 {
+	const FString DirectoryFullPath = FPaths::ConvertRelativePathToFull(Directory);
+	
 	for (TMap<FDirectoryWithFlags, FDirectoryWatchRequestWindows*>::TIterator RequestIt(RequestMap); RequestIt; ++RequestIt)
 	{
-		if (RequestIt->Key.Key == Directory)
+		if (RequestIt->Key.Key == DirectoryFullPath)
 		{
 			// There should be no NULL entries in the map
 			FDirectoryWatchRequestWindows* Watch = RequestIt->Value;
