Console and Mobile Packaging fails on QAGame due to undefined symbols
-This fixes the compile error, the stub implementation has perhaps never worked but also never been used because this file is not usually compiled except in the editor and real implementaitons exist for all editor platforms.
#Synced-CL 48465700

diff --git a/Engine/Source/Developer/DirectoryWatcher/Private/DirectoryWatcherPrivate.h b/Engine/Source/Developer/DirectoryWatcher/Private/DirectoryWatcherPrivate.h
--- a/Engine/Source/Developer/DirectoryWatcher/Private/DirectoryWatcherPrivate.h
+++ b/Engine/Source/Developer/DirectoryWatcher/Private/DirectoryWatcherPrivate.h
@@ -16,5 +16,6 @@ DECLARE_LOG_CATEGORY_EXTERN(LogDirectoryWatcher, Log, All);
 #include "Linux/DirectoryWatchRequestLinux.h"
 #include "Linux/DirectoryWatcherLinux.h"
 #else
+#pragma message("FDirectoryWatcher is using DirectoryWatcherStub.h on this platform, which does nothing.")
 #include "DirectoryWatcherStub.h"
 #endif
diff --git a/Engine/Source/Developer/DirectoryWatcher/Private/DirectoryWatcherStub.h b/Engine/Source/Developer/DirectoryWatcher/Private/DirectoryWatcherStub.h
--- a/Engine/Source/Developer/DirectoryWatcher/Private/DirectoryWatcherStub.h
+++ b/Engine/Source/Developer/DirectoryWatcher/Private/DirectoryWatcherStub.h
@@ -9,8 +9,8 @@
 class FDirectoryWatcherStub : public IDirectoryWatcher
 {
 public:
-	FDirectoryWatcherStub();
-	virtual ~FDirectoryWatcherStub();
+	FDirectoryWatcherStub() = default;
+	virtual ~FDirectoryWatcherStub() = default;
 
 	virtual bool RegisterDirectoryChangedCallback_Handle (const FString& Directory, const FDirectoryChanged& InDelegate, FDelegateHandle& OutHandle, uint32 Flags) override { return false; }
 	virtual bool UnregisterDirectoryChangedCallback_Handle(const FString& Directory, FDelegateHandle InHandle) override { return false; }
