UE: In DirectoryWatcher, make a copy of the DirectoryWatchRequest list instead of iterating over RequestMap in Tick so that delegates can modify it while we are iterating over the requests

#Synced-CL 50906850

diff --git a/Engine/Source/Developer/DirectoryWatcher/Private/Linux/DirectoryWatchRequestLinux.cpp b/Engine/Source/Developer/DirectoryWatcher/Private/Linux/DirectoryWatchRequestLinux.cpp
--- a/Engine/Source/Developer/DirectoryWatcher/Private/Linux/DirectoryWatchRequestLinux.cpp
+++ b/Engine/Source/Developer/DirectoryWatcher/Private/Linux/DirectoryWatchRequestLinux.cpp
@@ -144,13 +144,20 @@ void FDirectoryWatchRequestLinux::ProcessNotifications(TMap<FString, FDirectoryW
 	checkf(IsInGameThread(), TEXT("INotify operations only support on main thread"));
 
 	ProcessAllINotifyChanges();
-
-	// Trigger any file change notification delegates
+	
+	// Create a list of the WatchRequests to process this tick. Do not iterate over RequestMap when calling ProcessPendingNotifications because
+	// calls to Register or Unregister requests while processing delegates would modify it while we are iterating over it
+	TArray<FDirectoryWatchRequestLinux*> WatchRequestsToProcessThisFrame;
+	WatchRequestsToProcessThisFrame.Reserve(RequestMap.Num());
 	for (auto MapIt = RequestMap.CreateConstIterator(); MapIt; ++MapIt)
 	{
-		FDirectoryWatchRequestLinux &WatchRequest = *MapIt.Value();
-
-		WatchRequest.ProcessPendingNotifications();
+		WatchRequestsToProcessThisFrame.Add(MapIt.Value());
+	}
+	
+	// Trigger any file change notification delegates
+	for (FDirectoryWatchRequestLinux* WatchRequest : WatchRequestsToProcessThisFrame)
+	{
+		WatchRequest->ProcessPendingNotifications();
 	}
 
 	DumpINotifyErrorDetails(RequestMap);
diff --git a/Engine/Source/Developer/DirectoryWatcher/Private/Mac/DirectoryWatcherMac.cpp b/Engine/Source/Developer/DirectoryWatcher/Private/Mac/DirectoryWatcherMac.cpp
--- a/Engine/Source/Developer/DirectoryWatcher/Private/Mac/DirectoryWatcherMac.cpp
+++ b/Engine/Source/Developer/DirectoryWatcher/Private/Mac/DirectoryWatcherMac.cpp
@@ -115,9 +115,18 @@ void FDirectoryWatcherMac::Tick( float DeltaSeconds )
 		RequestsPendingDelete.RemoveAt(RequestIdx);
 	}
 
-	// Trigger any file change notification delegates
+	// Create a list of the WatchRequests to process this tick. Do not iterate over RequestMap when calling ProcessPendingNotifications because
+	// calls to Register or Unregister requests while processing delegates would modify it while we are iterating over it
+	TArray<FDirectoryWatchRequestMac*> WatchRequestsToProcessThisFrame;
+	WatchRequestsToProcessThisFrame.Reserve(RequestMap.Num());
 	for (TMap<FString, FDirectoryWatchRequestMac*>::TConstIterator RequestIt(RequestMap); RequestIt; ++RequestIt)
 	{
-		RequestIt.Value()->ProcessPendingNotifications();
+		WatchRequestsToProcessThisFrame.Add(RequestIt.Value());
+	}
+	
+	// Trigger any file change notification delegates
+	for (FDirectoryWatchRequestMac* WatchRequest : WatchRequestsToProcessThisFrame)
+	{
+		WatchRequest->ProcessPendingNotifications();
 	}
 }
diff --git a/Engine/Source/Developer/DirectoryWatcher/Private/Windows/DirectoryWatcherWindows.cpp b/Engine/Source/Developer/DirectoryWatcher/Private/Windows/DirectoryWatcherWindows.cpp
--- a/Engine/Source/Developer/DirectoryWatcher/Private/Windows/DirectoryWatcherWindows.cpp
+++ b/Engine/Source/Developer/DirectoryWatcher/Private/Windows/DirectoryWatcherWindows.cpp
@@ -120,6 +120,11 @@ void FDirectoryWatcherWindows::Tick( float DeltaSeconds )
 	// OS before we process the callbacks we received.
 	bool bHasPolled = false;
 
+	// Create a list of the WatchRequests to process this tick. Do not iterate over RequestMap when calling ProcessPendingNotifications because
+	// calls to Register or Unregister requests while processing delegates would modify it while we are iterating over it
+	TArray<FDirectoryWatchRequestWindows*> WatchRequestsToProcessThisFrame;
+	WatchRequestsToProcessThisFrame.Reserve(RequestMap.Num());
+
 	for (TMap<FDirectoryWithFlags, FDirectoryWatchRequestWindows*>::TIterator RequestIt(RequestMap);
 		RequestIt; ++RequestIt)
 	{
@@ -130,17 +135,22 @@ void FDirectoryWatcherWindows::Tick( float DeltaSeconds )
 		}
 		else
 		{
-			if (!bHasPolled)
-			{
-				// Trigger any OS notifications that are queued up in Windows's MessageQueue. To do this we need to put the
-				// current thread into an alertable state, which we can do by calling SleepEx(0).
-				SleepEx(0 /* TimeoutMs. 0 means poll notifies and yield to waiting threads but otherwise return immediately */,
-						1 /* bAlertable: the notifications will be processed */);
-				bHasPolled = true;
-			}
-			// Pass on the unreal notifications for the request if it received an OS notification.
-			RequestIt.Value()->ProcessPendingNotifications();
+			WatchRequestsToProcessThisFrame.Add(RequestIt.Value());
+		}
+	}
+
+	for (FDirectoryWatchRequestWindows* WatchRequest : WatchRequestsToProcessThisFrame)
+	{
+		if (!bHasPolled)
+		{
+			// Trigger any OS notifications that are queued up in Windows's MessageQueue. To do this we need to put the
+			// current thread into an alertable state, which we can do by calling SleepEx(0).
+			SleepEx(0 /* TimeoutMs. 0 means poll notifies and yield to waiting threads but otherwise return immediately */,
+					1 /* bAlertable: the notifications will be processed */);
+			bHasPolled = true;
 		}
+		// Pass on the unreal notifications for the request if it received an OS notification.
+		WatchRequest->ProcessPendingNotifications();
 	}
 
 	// Delete any stale or invalid requests, either the ones that we found during tick or ones that were
