Unshelved from pending changelist '50512592':

PR #14371: Removed redundant function call, speeding up the automatic asset reimporter  by 2x

#rnx
#Synced-CL 51544706

diff --git a/Engine/Source/Developer/DirectoryWatcher/Private/FileCache.cpp b/Engine/Source/Developer/DirectoryWatcher/Private/FileCache.cpp
--- a/Engine/Source/Developer/DirectoryWatcher/Private/FileCache.cpp
+++ b/Engine/Source/Developer/DirectoryWatcher/Private/FileCache.cpp
@@ -634,7 +634,10 @@ void FFileCache::DiffDirtyFiles(TMap<FImmutableString, FFileData>& InDirtyFiles,
 		
 		const auto* CachedState = CachedDirectoryState.Files.Find(File);
 
-		const bool bFileExists = InFileSystemState ? InFileSystemState->Files.Find(File) != nullptr : PlatformFile.FileExists(*AbsoluteFilename);
+		//Retrieve the timestamp directly instead of checking the existence of the file beforehand to save time
+		//If the file doesn't exist the function will return FDateTime::MinValue to indicate this
+		const FDateTime TimeStamp = InFileSystemState ? FDateTime::MinValue() : FileManager.GetTimeStamp(*AbsoluteFilename);
+		const bool bFileExists = InFileSystemState ? InFileSystemState->Files.Find(File) != nullptr : TimeStamp != FDateTime::MinValue();
 		if (bFileExists)
 		{
 			FFileData FileData;
@@ -645,7 +648,7 @@ void FFileCache::DiffDirtyFiles(TMap<FImmutableString, FFileData>& InDirtyFiles,
 			else
 			{
 				// The dirty file timestamp is the time that the file was dirtied, not necessarily its modification time
-				FileData = FFileData(FileManager.GetTimeStamp(*AbsoluteFilename), Pair.Value.FileHash);
+				FileData = FFileData(TimeStamp, Pair.Value.FileHash);
 			}
 
 			if (Config.bRequireFileHashes && !FileData.FileHash.IsValid())
